Object Mixer & Light Mixer:

-Move modules to plugins

diff --git a/Engine/Plugins/Editor/ObjectMixer/LightMixer/LightMixer.uplugin b/Engine/Plugins/Editor/ObjectMixer/LightMixer/LightMixer.uplugin
new file mode 100644
--- /dev/null
+++ b/Engine/Plugins/Editor/ObjectMixer/LightMixer/LightMixer.uplugin
@@ -0,0 +1,32 @@
+{
+	"FileVersion": 1,
+	"Version": 1,
+	"VersionName": "1.0",
+	"FriendlyName": "Light Mixer",
+	"Description": "Edit any properties of scene lights in a spreadsheet format!",
+	"Category": "Editor",
+	"CreatedBy": "",
+	"CreatedByURL": "",
+	"DocsURL": "",
+	"EnabledByDefault": true,
+	"MarketplaceURL": "",
+	"SupportURL": "",
+	"CanContainContent": true,
+	"IsBetaVersion": true,
+	"IsExperimentalVersion": false,
+	"Installed": true,
+	"Hidden": true,
+	"Plugins": [
+		{
+			"Name": "ObjectMixer",
+			"Enabled": true
+		}
+	],
+	"Modules": [
+		{
+			"Name": "LightMixer",
+			"Type": "Editor",
+			"LoadingPhase": "Default"
+		}
+	]
+}
diff --git a/Engine/Plugins/Editor/ObjectMixer/LightMixer/Source/LightMixer/LightMixer.build.cs b/Engine/Plugins/Editor/ObjectMixer/LightMixer/Source/LightMixer/LightMixer.build.cs
new file mode 100644
--- /dev/null
+++ b/Engine/Plugins/Editor/ObjectMixer/LightMixer/Source/LightMixer/LightMixer.build.cs
@@ -0,0 +1,48 @@
+// Copyright Epic Games, Inc. All Rights Reserved.
+
+using System.IO;
+using UnrealBuildTool;
+
+public class LightMixer : ModuleRules
+{
+	public LightMixer(ReadOnlyTargetRules Target) : base(Target)
+	{
+		PCHUsage = ModuleRules.PCHUsageMode.UseExplicitOrSharedPCHs;
+		
+		PublicIncludePaths.Add(Path.Combine(ModuleDirectory, "Public"));
+
+		PublicDependencyModuleNames.AddRange(
+			new string[]
+			{
+				"Core",
+				"OutputLog",
+				"ObjectMixerEditor",
+				"PropertyEditor"
+			}
+		);
+			
+		
+		PrivateDependencyModuleNames.AddRange(
+			new string[]
+			{
+				"AssetRegistry",
+				"AssetTools",
+				"CoreUObject",
+				"ContentBrowser",
+				"Engine",
+				"EditorStyle",
+				"EditorWidgets",
+				"InputCore",
+				"Kismet",
+				"ObjectMixerEditor",
+				"Projects",
+				"Slate",
+				"SlateCore",
+				"ToolMenus", 
+				"ToolWidgets",
+				"UnrealEd",
+				"WorkspaceMenuStructure"
+			}
+		);
+	}
+}
diff --git a/Engine/Plugins/Editor/ObjectMixer/LightMixer/Source/LightMixer/Private/LightMixerModule.cpp b/Engine/Plugins/Editor/ObjectMixer/LightMixer/Source/LightMixer/Private/LightMixerModule.cpp
new file mode 100644
--- /dev/null
+++ b/Engine/Plugins/Editor/ObjectMixer/LightMixer/Source/LightMixer/Private/LightMixerModule.cpp
@@ -0,0 +1,90 @@
+// Copyright Epic Games, Inc. All Rights Reserved.
+
+#include "LightMixerModule.h"
+
+#include "LightMixerObjectFilter.h"
+#include "LightMixerProjectSettings.h"
+#include "LightMixerStyle.h"
+
+#include "ISettingsModule.h"
+#include "ISettingsSection.h"
+#include "Modules/ModuleManager.h"
+
+#define LOCTEXT_NAMESPACE "FLightMixerEditorModule"
+
+IMPLEMENT_MODULE(FLightMixerModule, LightMixer)
+
+void FLightMixerModule::StartupModule()
+{
+	FLightMixerStyle::Initialize();
+	
+	FCoreDelegates::OnPostEngineInit.AddRaw(this, &FLightMixerModule::Initialize);
+}
+
+void FLightMixerModule::ShutdownModule()
+{
+	FLightMixerStyle::Shutdown();
+
+	Teardown();
+}
+
+FLightMixerModule& FLightMixerModule::Get()
+{
+	return FModuleManager::LoadModuleChecked< FLightMixerModule >("LightMixer");
+}
+
+void FLightMixerModule::OpenProjectSettings()
+{
+	FModuleManager::LoadModuleChecked<ISettingsModule>("Settings")
+		.ShowViewer("Project", "Editor", "Light Mixer");
+}
+
+void FLightMixerModule::Initialize()
+{
+	FObjectMixerEditorModule::Initialize();
+	
+	DefaultFilterClass = ULightMixerObjectFilter::StaticClass();	
+}
+
+FName FLightMixerModule::GetModuleName()
+{
+	return "LightMixer";
+}
+
+void FLightMixerModule::SetupMenuItemVariables()
+{
+	TabLabel = LOCTEXT("LightMixerTabLabel", "Light Mixer");
+
+	MenuItemName = LOCTEXT("OpenLightMixerEditorMenuItem", "Light Mixer");
+	MenuItemIcon =
+		FSlateIcon(FLightMixerStyle::Get().GetStyleSetName(), "LightMixer.ToolbarButton", "LightMixer.ToolbarButton.Small");
+	MenuItemTooltip = LOCTEXT("OpenLightMixerEditorTooltip", "Open Light Mixer");
+}
+
+FName FLightMixerModule::GetTabSpawnerId()
+{
+	return "LightMixerToolkit";
+}
+
+void FLightMixerModule::RegisterProjectSettings() const
+{
+	if (ISettingsModule* SettingsModule = FModuleManager::GetModulePtr<ISettingsModule>("Settings"))
+	{
+		// User Project Settings
+		const TSharedPtr<ISettingsSection> ProjectSettingsSectionPtr = SettingsModule->RegisterSettings(
+			"Project", "Editor", "Light Mixer",
+			LOCTEXT("LightMixerSettingsDisplayName", "Light Mixer"),
+			LOCTEXT("LightMixerSettingsDescription", "Configure Light Mixer user settings"),
+			GetMutableDefault<ULightMixerProjectSettings>());
+	}
+}
+
+void FLightMixerModule::UnregisterProjectSettings() const
+{
+	if (ISettingsModule* SettingsModule = FModuleManager::GetModulePtr<ISettingsModule>("Settings"))
+	{
+		SettingsModule->UnregisterSettings("Project", "Editor", "Light Mixer");
+	}
+}
+
+#undef LOCTEXT_NAMESPACE
diff --git a/Engine/Plugins/Editor/ObjectMixer/LightMixer/Source/LightMixer/Private/LightMixerStyle.cpp b/Engine/Plugins/Editor/ObjectMixer/LightMixer/Source/LightMixer/Private/LightMixerStyle.cpp
new file mode 100644
--- /dev/null
+++ b/Engine/Plugins/Editor/ObjectMixer/LightMixer/Source/LightMixer/Private/LightMixerStyle.cpp
@@ -0,0 +1,90 @@
+// Copyright Epic Games, Inc. All Rights Reserved.
+
+#include "LightMixerStyle.h"
+
+#include "Styling/AppStyle.h"
+#include "Framework/Application/SlateApplication.h"
+#include "Interfaces/IPluginManager.h"
+#include "Styling/CoreStyle.h"
+#include "Styling/SlateStyleRegistry.h"
+#include "Styling/SlateTypes.h"
+#include "Styling/StyleColors.h"
+
+TSharedPtr<FSlateStyleSet> FLightMixerStyle::StyleInstance = nullptr;
+
+void FLightMixerStyle::Initialize()
+{
+	if (!StyleInstance.IsValid())
+	{
+		StyleInstance = Create();
+		FSlateStyleRegistry::RegisterSlateStyle(*StyleInstance);
+	}
+}
+
+void FLightMixerStyle::Shutdown()
+{
+	FSlateStyleRegistry::UnRegisterSlateStyle(*StyleInstance);
+	ensure(StyleInstance.IsUnique());
+	StyleInstance.Reset();
+}
+
+void FLightMixerStyle::ReloadTextures()
+{
+	FSlateApplication::Get().GetRenderer()->ReloadTextureResources();
+}
+
+const ISlateStyle& FLightMixerStyle::Get()
+{
+	return *StyleInstance;
+}
+
+const FName& FLightMixerStyle::GetStyleSetName() const
+{
+	static FName StyleName(TEXT("LightMixer"));
+	return StyleName;
+}
+
+const FSlateBrush* FLightMixerStyle::GetBrush(const FName PropertyName, const ANSICHAR* Specifier, const ISlateStyle* RequestingStyle) const
+{
+	return StyleInstance->GetBrush(PropertyName, Specifier);
+}
+
+#define IMAGE_BRUSH( RelativePath, ... ) FSlateImageBrush( Style->RootToContentDir( RelativePath, TEXT(".png") ), __VA_ARGS__ )
+#define BOX_BRUSH( RelativePath, ... ) FSlateBoxBrush( Style->RootToContentDir( RelativePath, TEXT(".png") ), __VA_ARGS__ )
+#define BORDER_BRUSH( RelativePath, ... ) FSlateBorderBrush( Style->RootToContentDir( RelativePath, TEXT(".png") ), __VA_ARGS__ )
+#define IMAGE_BRUSH_SVG( RelativePath, ... ) FSlateVectorImageBrush( Style->RootToContentDir(RelativePath, TEXT(".svg") ), __VA_ARGS__)
+#define IMAGE_PLUGIN_BRUSH_SVG( PluginName, RelativePath, ... ) FSlateVectorImageBrush( FLightMixerEditorStyle::GetExternalPluginContent(PluginName, RelativePath, ".svg"), __VA_ARGS__)
+
+const FVector2D Icon64x64(64.f, 64.f);
+const FVector2D Icon40x40(40.0f, 40.0f);
+const FVector2D Icon20x20(20.0f, 20.0f);
+const FVector2D Icon16x16(16.0f, 16.0f);
+const FVector2D Icon12x12(12.0f, 12.0f);
+const FVector2D Icon8x8(8.f, 8.f);
+
+FString FLightMixerStyle::GetExternalPluginContent(const FString& PluginName, const FString& RelativePath, const ANSICHAR* Extension)
+{
+	FString ContentDir = IPluginManager::Get().FindPlugin(PluginName)->GetBaseDir() / RelativePath + Extension;
+	return ContentDir;
+}
+
+TSharedRef< FSlateStyleSet > FLightMixerStyle::Create()
+{
+	TSharedRef<FSlateStyleSet> Style = MakeShared<FSlateStyleSet>("LightMixer");
+	
+	//Style->SetContentRoot(FPaths::Combine(Plugin->GetBaseDir(), TEXT("Resources")));
+
+	Style->SetCoreContentRoot(FPaths::EngineContentDir() / TEXT("Slate"));
+
+	// Icons
+	Style->Set("LightMixer.ToolbarButton", new IMAGE_BRUSH_SVG("Icons/LightMixer", Icon40x40));
+	Style->Set("LightMixer.ToolbarButton.Small", new IMAGE_BRUSH_SVG("Icons/LightMixer", Icon20x20));
+	
+	return Style;
+}
+
+#undef IMAGE_BRUSH
+#undef BOX_BRUSH
+#undef BORDER_BRUSH
+#undef IMAGE_BRUSH_SVG
+#undef IMAGE_PLUGIN_BRUSH_SVG
diff --git a/Engine/Plugins/Editor/ObjectMixer/LightMixer/Source/LightMixer/Private/LightMixerStyle.h b/Engine/Plugins/Editor/ObjectMixer/LightMixer/Source/LightMixer/Private/LightMixerStyle.h
new file mode 100644
--- /dev/null
+++ b/Engine/Plugins/Editor/ObjectMixer/LightMixer/Source/LightMixer/Private/LightMixerStyle.h
@@ -0,0 +1,36 @@
+// Copyright Epic Games, Inc. All Rights Reserved.
+
+#pragma once
+
+#include "Styling/SlateStyle.h"
+
+class FLightMixerStyle : FSlateStyleSet
+{
+public:
+
+	static void Initialize();
+
+	static void Shutdown();
+
+	static void ReloadTextures();
+
+	static const ISlateStyle& Get();
+
+	virtual const FName& GetStyleSetName() const override;
+
+	virtual const FSlateBrush* GetBrush(const FName PropertyName, const ANSICHAR* Specifier = nullptr, const ISlateStyle* RequestingStyle = nullptr) const override;
+
+	template< typename WidgetStyleType >
+	static const WidgetStyleType& GetWidgetStyle(FName PropertyName, const ANSICHAR* Specifier = nullptr)
+	{
+		return StyleInstance->GetWidgetStyle<WidgetStyleType>(PropertyName, Specifier);
+	}
+
+private:
+
+	static FString GetExternalPluginContent(const FString& PluginName, const FString& RelativePath, const ANSICHAR* Extension);
+	
+	static TSharedRef<FSlateStyleSet> Create();
+
+	static TSharedPtr<FSlateStyleSet> StyleInstance;
+};
diff --git a/Engine/Plugins/Editor/ObjectMixer/LightMixer/Source/LightMixer/Public/LightMixerModule.h b/Engine/Plugins/Editor/ObjectMixer/LightMixer/Source/LightMixer/Public/LightMixerModule.h
new file mode 100644
--- /dev/null
+++ b/Engine/Plugins/Editor/ObjectMixer/LightMixer/Source/LightMixer/Public/LightMixerModule.h
@@ -0,0 +1,28 @@
+// Copyright Epic Games, Inc. All Rights Reserved.
+
+#pragma once
+
+#include "ObjectMixerEditorModule.h"
+
+class LIGHTMIXER_API FLightMixerModule : public FObjectMixerEditorModule
+{
+public:
+
+	//~ Begin IModuleInterface Interface
+	virtual void StartupModule() override;
+	virtual void ShutdownModule() override;
+	//~ End IModuleInterface Interface
+
+	static FLightMixerModule& Get();
+	
+	static void OpenProjectSettings();
+	
+	//~ Begin FObjectMixerEditorModule overrides
+	virtual void Initialize() override;
+	virtual FName GetModuleName() override;
+	virtual void SetupMenuItemVariables() override;
+	virtual FName GetTabSpawnerId() override;
+	virtual void RegisterProjectSettings() const override;
+	virtual void UnregisterProjectSettings() const override;
+	//~ End FObjectMixerEditorModule overrides
+};
diff --git a/Engine/Plugins/Editor/ObjectMixer/LightMixer/Source/LightMixer/Public/LightMixerObjectFilter.h b/Engine/Plugins/Editor/ObjectMixer/LightMixer/Source/LightMixer/Public/LightMixerObjectFilter.h
new file mode 100644
--- /dev/null
+++ b/Engine/Plugins/Editor/ObjectMixer/LightMixer/Source/LightMixer/Public/LightMixerObjectFilter.h
@@ -0,0 +1,61 @@
+// Copyright Epic Games, Inc. All Rights Reserved.
+
+#pragma once
+
+#include "ObjectFilter/ObjectMixerEditorObjectFilter.h"
+
+#include "Engine/Light.h"
+#include "GameFramework/Actor.h"
+
+#include "LightMixerObjectFilter.generated.h"
+
+UCLASS(BlueprintType, EditInlineNew)
+class LIGHTMIXER_API ULightMixerObjectFilter : public UObjectMixerObjectFilter
+{
+	GENERATED_BODY()
+public:
+	
+	virtual TSet<UClass*> GetObjectClassesToFilter() const override
+	{
+		return
+		{
+			ULightComponent::StaticClass()
+		};
+	}
+
+	virtual TSet<TSubclassOf<AActor>> GetObjectClassesToPlace() const override
+	{
+		return
+		{
+			ALight::StaticClass()
+		};
+	}
+
+	virtual TSet<FName> GetColumnsToShowByDefault() const override
+	{
+		return 
+		{
+			"Intensity", "LightColor", "AttenuationRadius"
+		};
+	}
+
+	virtual TSet<FName> GetForceAddedColumns() const override
+	{
+		return {"LightColor"};
+	}
+
+	virtual bool ShouldIncludeUnsupportedProperties() const override
+	{
+		return false;
+	}
+
+	virtual EObjectMixerInheritanceInclusionOptions GetObjectMixerPropertyInheritanceInclusionOptions() const override
+	{
+		return EObjectMixerInheritanceInclusionOptions::IncludeAllParentsAndChildren;
+	}
+
+	virtual EObjectMixerInheritanceInclusionOptions GetObjectMixerPlacementClassInclusionOptions() const override
+	{
+		return EObjectMixerInheritanceInclusionOptions::IncludeAllChildren;
+	}
+};
diff --git a/Engine/Plugins/Editor/ObjectMixer/LightMixer/Source/LightMixer/Public/LightMixerProjectSettings.h b/Engine/Plugins/Editor/ObjectMixer/LightMixer/Source/LightMixer/Public/LightMixerProjectSettings.h
new file mode 100644
--- /dev/null
+++ b/Engine/Plugins/Editor/ObjectMixer/LightMixer/Source/LightMixer/Public/LightMixerProjectSettings.h
@@ -0,0 +1,22 @@
+// Copyright Epic Games, Inc. All Rights Reserved.
+
+#pragma once
+
+#include "LightMixerProjectSettings.generated.h"
+
+UCLASS(config = ObjectMixer, defaultconfig)
+class LIGHTMIXER_API ULightMixerProjectSettings : public UObject
+{
+	GENERATED_BODY()
+public:
+	
+	ULightMixerProjectSettings(const FObjectInitializer& ObjectInitializer)
+	{}
+
+	/**
+	 * If true, the Object Mixer menu item will be disabled and removed.
+	 * This is useful if you only want to use Light Mixer and to avoid cluttering the menus.
+	 */
+	UPROPERTY(Config, EditAnywhere, Category="Light Mixer", meta = (ConfigRestartRequired = true))
+	bool bHideObjectMixerMenuItem = false;
+};
diff --git a/Engine/Plugins/Editor/ObjectMixer/ObjectMixer/ObjectMixer.uplugin b/Engine/Plugins/Editor/ObjectMixer/ObjectMixer/ObjectMixer.uplugin
new file mode 100644
--- /dev/null
+++ b/Engine/Plugins/Editor/ObjectMixer/ObjectMixer/ObjectMixer.uplugin
@@ -0,0 +1,26 @@
+{
+	"FileVersion": 1,
+	"Version": 1,
+	"VersionName": "1.0",
+	"FriendlyName": "Object Mixer",
+	"Description": "Edit any properties of scene objects in a spreadsheet format!",
+	"Category": "Editor",
+	"CreatedBy": "",
+	"CreatedByURL": "",
+	"DocsURL": "",
+	"EnabledByDefault": true,
+	"MarketplaceURL": "",
+	"SupportURL": "",
+	"CanContainContent": true,
+	"IsBetaVersion": true,
+	"IsExperimentalVersion": false,
+	"Installed": true,
+	"Hidden": true,
+	"Modules": [
+		{
+			"Name": "ObjectMixerEditor",
+			"Type": "Editor",
+			"LoadingPhase": "Default"
+		}
+	]
+}
diff --git a/Engine/Plugins/Editor/ObjectMixer/ObjectMixer/Source/ObjectMixer/ObjectMixerEditor.build.cs b/Engine/Plugins/Editor/ObjectMixer/ObjectMixer/Source/ObjectMixer/ObjectMixerEditor.build.cs
new file mode 100644
--- /dev/null
+++ b/Engine/Plugins/Editor/ObjectMixer/ObjectMixer/Source/ObjectMixer/ObjectMixerEditor.build.cs
@@ -0,0 +1,48 @@
+// Copyright Epic Games, Inc. All Rights Reserved.
+
+using System.IO;
+using UnrealBuildTool;
+
+public class ObjectMixerEditor : ModuleRules
+{
+	public ObjectMixerEditor(ReadOnlyTargetRules Target) : base(Target)
+	{
+		PCHUsage = ModuleRules.PCHUsageMode.UseExplicitOrSharedPCHs;
+		
+		PublicIncludePaths.Add(Path.Combine(ModuleDirectory, "Public"));
+
+		PublicDependencyModuleNames.AddRange(
+			new string[]
+			{
+				"Core",
+				"OutputLog",
+				"PropertyEditor"
+			}
+		);
+			
+		
+		PrivateDependencyModuleNames.AddRange(
+			new string[]
+			{
+				"AssetRegistry",
+				"AssetTools",
+				"CoreUObject",
+				"ContentBrowser",
+				"Engine",
+				"EditorStyle",
+				"EditorWidgets",
+				"InputCore",
+				"Kismet",
+				"Projects",
+				"SceneOutliner",
+				"Slate",
+				"SlateCore",
+				"ToolMenus", 
+				"ToolWidgets",
+				"TypedElementFramework",
+				"UnrealEd",
+				"WorkspaceMenuStructure"
+			}
+		);
+	}
+}
diff --git a/Engine/Plugins/Editor/ObjectMixer/ObjectMixer/Source/ObjectMixer/Private/ObjectFilter/ObjectMixerEditorObjectFilter.cpp b/Engine/Plugins/Editor/ObjectMixer/ObjectMixer/Source/ObjectMixer/Private/ObjectFilter/ObjectMixerEditorObjectFilter.cpp
new file mode 100644
--- /dev/null
+++ b/Engine/Plugins/Editor/ObjectMixer/ObjectMixer/Source/ObjectMixer/Private/ObjectFilter/ObjectMixerEditorObjectFilter.cpp
@@ -0,0 +1,217 @@
+// Copyright Epic Games, Inc. All Rights Reserved.
+
+#include "ObjectFilter/ObjectMixerEditorObjectFilter.h"
+
+#include "GameFramework/Actor.h"
+#include "Kismet2/ComponentEditorUtils.h"
+
+FText UObjectMixerObjectFilter::GetRowDisplayName(UObject* InObject, EObjectMixerTreeViewMode InViewMode) const
+{
+	if (IsValid(InObject))
+	{
+		if (InViewMode == EObjectMixerTreeViewMode::Flat || InViewMode == EObjectMixerTreeViewMode::Folder)
+		{
+			if (const AActor* Outer = InObject->GetTypedOuter<AActor>())
+			{
+				return FText::Format(INVTEXT("{0} ({1})"), FText::FromString(Outer->GetActorLabel()), FText::FromString(InObject->GetName()));
+			}
+		}
+
+		if (const AActor* AsActor = Cast<AActor>(InObject))
+		{
+			return FText::FromString(AsActor->GetActorLabel());
+		}
+		else if (const UActorComponent* AsActorComponent = Cast<UActorComponent>(InObject))
+		{
+			const FName ComponentName = FComponentEditorUtils::FindVariableNameGivenComponentInstance(AsActorComponent);
+			if (ComponentName != NAME_None)
+			{
+				if (!ComponentName.IsEqual(AsActorComponent->GetFName()))
+				{
+					return
+					   FText::Format(INVTEXT("{0} ({1})"),
+						   FText::FromName(ComponentName),
+						   FText::FromName(AsActorComponent->GetFName())
+					   );
+				}
+
+				return FText::FromName(ComponentName);
+			}
+		}
+		
+		return FText::FromString(InObject->GetName());
+	}
+
+	return FText::GetEmpty();
+}
+
+bool UObjectMixerObjectFilter::GetRowEditorVisibility(UObject* InObject, EObjectMixerTreeViewMode InViewMode) const
+{
+	if (IsValid(InObject))
+	{
+		TObjectPtr<AActor> Actor = Cast<AActor>(InObject);
+	
+		if (!Actor)
+		{
+			Actor = InObject->GetTypedOuter<AActor>();
+		}
+
+		return Actor ? !Actor->IsTemporarilyHiddenInEditor() : false;
+	}
+
+	return false;
+}
+
+void UObjectMixerObjectFilter::OnSetRowEditorVisibility(UObject* InObject, bool bNewIsVisible, EObjectMixerTreeViewMode InViewMode) const
+{
+	if (IsValid(InObject))
+	{
+		TObjectPtr<AActor> Actor = Cast<AActor>(InObject);
+	
+		if (!Actor)
+		{
+			Actor = InObject->GetTypedOuter<AActor>();
+		}
+
+		if (Actor)
+		{
+			Actor->SetIsTemporarilyHiddenInEditor(!bNewIsVisible);
+		}
+	}
+}
+
+TSet<FName> UObjectMixerObjectFilter::GetColumnsToShowByDefault() const
+{
+	return {};
+}
+
+TSet<FName> UObjectMixerObjectFilter::GetColumnsToExclude() const
+{
+	return {};
+}
+
+EObjectMixerInheritanceInclusionOptions UObjectMixerObjectFilter::GetObjectMixerPropertyInheritanceInclusionOptions() const
+{
+	return EObjectMixerInheritanceInclusionOptions::None;
+}
+
+EObjectMixerInheritanceInclusionOptions UObjectMixerObjectFilter::GetObjectMixerPlacementClassInclusionOptions() const
+{
+	return EObjectMixerInheritanceInclusionOptions::None;
+}
+
+bool UObjectMixerObjectFilter::ShouldIncludeUnsupportedProperties() const
+{
+	return false;
+}
+
+TSet<UClass*> UObjectMixerObjectFilter::GetParentAndChildClassesFromSpecifiedClasses(
+	const TSet<UClass*>& InSpecifiedClasses, EObjectMixerInheritanceInclusionOptions Options)
+{
+	// 'None' means we only want the specified classes
+	if (Options == EObjectMixerInheritanceInclusionOptions::None)
+	{
+		return InSpecifiedClasses;
+	}
+	
+	TSet<UClass*> ReturnValue;
+
+	auto GetChildClassesLambda = [](const UClass* Class, TSet<UClass*>& OutReturnValue, const bool bRecursive)
+	{
+		TArray<UClass*> DerivedClasses;
+		GetDerivedClasses(Class, DerivedClasses, bRecursive);
+		OutReturnValue.Append(DerivedClasses);
+	};
+
+	for (UClass* Class : InSpecifiedClasses)
+	{
+		ReturnValue.Add(Class);
+
+		// Super Classes
+
+			// Immediate only
+		if (Options == EObjectMixerInheritanceInclusionOptions::IncludeOnlyImmediateParent ||
+			Options == EObjectMixerInheritanceInclusionOptions::IncludeOnlyImmediateParentAndChildren ||
+			Options == EObjectMixerInheritanceInclusionOptions::IncludeOnlyImmediateParentAndAllChildren)
+		{
+			if (UClass* Super = Class->GetSuperClass())
+			{
+				ReturnValue.Add(Super);
+			}
+		}
+			// All Parents
+		if (Options == EObjectMixerInheritanceInclusionOptions::IncludeAllParents ||
+			Options == EObjectMixerInheritanceInclusionOptions::IncludeAllParentsAndChildren ||
+			Options == EObjectMixerInheritanceInclusionOptions::IncludeAllParentsAndOnlyImmediateChildren)
+		{
+			if (UClass* Super = Class->GetSuperClass())
+			{
+				while (Super)
+				{
+					ReturnValue.Add(Super);
+					Super = Super->GetSuperClass();
+				}
+			}
+		}
+
+		// Child Classes
+
+			// Immediate only
+		if (Options == EObjectMixerInheritanceInclusionOptions::IncludeOnlyImmediateChildren ||
+			Options == EObjectMixerInheritanceInclusionOptions::IncludeOnlyImmediateParentAndChildren ||
+			Options == EObjectMixerInheritanceInclusionOptions::IncludeAllParentsAndOnlyImmediateChildren)
+		{
+			GetChildClassesLambda(Class, ReturnValue, false);
+		}
+			// All Children
+		if (Options == EObjectMixerInheritanceInclusionOptions::IncludeAllChildren ||
+			Options == EObjectMixerInheritanceInclusionOptions::IncludeAllParentsAndChildren ||
+			Options == EObjectMixerInheritanceInclusionOptions::IncludeOnlyImmediateParentAndAllChildren)
+		{
+			GetChildClassesLambda(Class, ReturnValue, true);
+		}
+	}
+
+	return ReturnValue;
+}
+
+TSet<UClass*> UObjectMixerObjectFilter::GetParentAndChildClassesFromSpecifiedClasses(
+	const TSet<TSubclassOf<AActor>>& InSpecifiedClasses, EObjectMixerInheritanceInclusionOptions Options)
+{
+	TSet<UClass*> AsUClasses;
+	
+	for (TSubclassOf<AActor> ActorClass : InSpecifiedClasses)
+	{
+		AsUClasses.Add(ActorClass);
+	}
+
+	return GetParentAndChildClassesFromSpecifiedClasses(AsUClasses, Options);
+}
+
+EFieldIterationFlags UObjectMixerObjectFilter::GetDesiredFieldIterationFlags(const bool bIncludeInheritedProperties)
+{
+	return bIncludeInheritedProperties ? EFieldIterationFlags::IncludeSuper : EFieldIterationFlags::Default;
+}
+
+TSet<FName> UObjectMixerObjectFilter::GenerateIncludeListFromExcludeList(const TSet<FName>& ExcludeList) const
+{
+	TSet<FName> IncludeList;
+
+	const EObjectMixerInheritanceInclusionOptions Options =
+		GetObjectMixerPropertyInheritanceInclusionOptions();
+	TSet<UClass*> SpecifiedClasses =
+		GetParentAndChildClassesFromSpecifiedClasses(GetObjectClassesToFilter(), Options);
+	
+	for (const UClass* Class : SpecifiedClasses)
+	{
+		for (TFieldIterator<FProperty> FieldIterator(Class); FieldIterator; ++FieldIterator)
+		{
+			if (const FProperty* Property = *FieldIterator)
+			{
+				IncludeList.Add(Property->GetFName());
+			}
+		}
+	}
+
+	return IncludeList.Difference(ExcludeList);
+}
diff --git a/Engine/Plugins/Editor/ObjectMixer/ObjectMixer/Source/ObjectMixer/Private/ObjectMixerEditorModule.cpp b/Engine/Plugins/Editor/ObjectMixer/ObjectMixer/Source/ObjectMixer/Private/ObjectMixerEditorModule.cpp
new file mode 100644
--- /dev/null
+++ b/Engine/Plugins/Editor/ObjectMixer/ObjectMixer/Source/ObjectMixer/Private/ObjectMixerEditorModule.cpp
@@ -0,0 +1,334 @@
+// Copyright Epic Games, Inc. All Rights Reserved.
+
+#include "ObjectMixerEditorModule.h"
+
+#include "ObjectMixerEditorLog.h"
+#include "ObjectMixerEditorProjectSettings.h"
+#include "ObjectMixerEditorStyle.h"
+#include "Views/MainPanel/ObjectMixerEditorMainPanel.h"
+
+#include "Editor.h"
+#include "Framework/Docking/TabManager.h"
+#include "ISettingsModule.h"
+#include "LevelEditor.h"
+#include "ObjectMixerEditorSerializedData.h"
+#include "Misc/TransactionObjectEvent.h"
+#include "Modules/ModuleManager.h"
+#include "ToolMenus.h"
+#include "WorkspaceMenuStructure.h"
+#include "WorkspaceMenuStructureModule.h"
+
+const FName FObjectMixerEditorModule::BaseObjectMixerModuleName("ObjectMixerEditor");
+
+#define LOCTEXT_NAMESPACE "FObjectMixerEditorModule"
+
+IMPLEMENT_MODULE(FObjectMixerEditorModule, ObjectMixerEditor)
+
+void FObjectMixerEditorModule::StartupModule()
+{
+	FObjectMixerEditorStyle::Initialize();
+
+	// In the future, Object Mixer and Light Mixer tabs may go into an Object Mixer group
+	//RegisterMenuGroup();
+	
+	FCoreDelegates::OnPostEngineInit.AddRaw(this, &FObjectMixerEditorModule::Initialize);
+}
+
+void FObjectMixerEditorModule::ShutdownModule()
+{
+	FObjectMixerEditorStyle::Shutdown();
+
+	UnregisterMenuGroup();
+	
+	Teardown();
+}
+
+void FObjectMixerEditorModule::Initialize()
+{
+	BindDelegates();
+	
+	SetupMenuItemVariables();
+	
+	RegisterTabSpawner();
+	RegisterProjectSettings();
+}
+
+void FObjectMixerEditorModule::Teardown()
+{
+	// Unbind Delegates
+	if (GEngine)
+	{
+		GEngine->OnLevelActorFolderChanged().RemoveAll(this);
+		GEngine->OnLevelActorListChanged().RemoveAll(this);
+		GEngine->OnActorFolderAdded().RemoveAll(this);
+		GEngine->OnLevelActorDetached().RemoveAll(this);
+		GEngine->OnLevelActorAttached().RemoveAll(this);
+		GEngine->OnActorFolderAdded().RemoveAll(this);
+		GEngine->OnActorFoldersUpdatedEvent().RemoveAll(this);
+	}
+
+	if (GEditor)
+	{
+		GEditor->OnBlueprintCompiled().RemoveAll(this);
+	}
+	
+	FEditorDelegates::MapChange.RemoveAll(this);
+	FEditorDelegates::PostUndoRedo.RemoveAll(this);
+	FCoreDelegates::OnPostEngineInit.RemoveAll(this);
+	FCoreUObjectDelegates::OnObjectTransacted.RemoveAll(this);
+
+	if (FLevelEditorModule* LevelEditorModule = FModuleManager::GetModulePtr<FLevelEditorModule>("LevelEditor"))
+	{
+		LevelEditorModule->OnComponentsEdited().RemoveAll(this);
+	}
+
+	for (FDelegateHandle& Delegate : DelegateHandles)
+	{
+		Delegate.Reset();
+	}
+	DelegateHandles.Empty();
+	
+	MainPanel.Reset();
+
+	UToolMenus::UnregisterOwner(this);
+	
+	UnregisterTabSpawner();
+	UnregisterProjectSettings();
+}
+
+FObjectMixerEditorModule& FObjectMixerEditorModule::Get()
+{
+	return FModuleManager::LoadModuleChecked< FObjectMixerEditorModule >("ObjectMixerEditor");
+}
+
+void FObjectMixerEditorModule::OpenProjectSettings()
+{
+	FModuleManager::LoadModuleChecked<ISettingsModule>("Settings")
+		.ShowViewer("Project", "Editor", "Object Mixer");
+}
+
+FName FObjectMixerEditorModule::GetModuleName()
+{
+	return "ObjectMixerEditor";
+}
+
+TSharedPtr<SWidget> FObjectMixerEditorModule::MakeObjectMixerDialog() const
+{
+	if (MainPanel.IsValid())
+	{
+		return MainPanel->GetOrCreateWidget();
+	}
+
+	return nullptr;
+}
+
+void FObjectMixerEditorModule::RequestRebuildList() const
+{
+	if (MainPanel.IsValid())
+	{
+		MainPanel->RequestRebuildList();
+	}
+}
+
+void FObjectMixerEditorModule::RefreshList() const
+{
+	if (MainPanel.IsValid())
+	{
+		MainPanel->RefreshList();
+	}
+}
+
+void FObjectMixerEditorModule::RegisterMenuGroup()
+{
+	WorkspaceGroup = WorkspaceMenu::GetMenuStructure().GetLevelEditorCategory()->AddGroup(
+		LOCTEXT("ObjectMixerMenuGroupItemName", "Object Mixer"), 
+		FSlateIcon(FObjectMixerEditorStyle::Get().GetStyleSetName(),
+			"ObjectMixer.ToolbarButton", "ObjectMixer.ToolbarButton.Small"));
+}
+
+void FObjectMixerEditorModule::UnregisterMenuGroup()
+{
+	if (WorkspaceGroup)
+	{
+		for (const TSharedRef<FWorkspaceItem>& ChildItem : WorkspaceGroup->GetChildItems())
+		{
+			WorkspaceGroup->RemoveItem(ChildItem);
+		}
+		
+		WorkspaceMenu::GetMenuStructure().GetLevelEditorCategory()->RemoveItem(WorkspaceGroup.ToSharedRef());
+		WorkspaceGroup.Reset();
+	}
+}
+
+void FObjectMixerEditorModule::SetupMenuItemVariables()
+{
+	TabLabel = LOCTEXT("ObjectMixerTabLabel", "Object Mixer");
+
+	MenuItemName = LOCTEXT("ObjectMixerEditorMenuItem", "Object Mixer");
+	MenuItemIcon =
+		FSlateIcon(FObjectMixerEditorStyle::Get().GetStyleSetName(), "ObjectMixer.ToolbarButton", "ObjectMixer.ToolbarButton.Small");
+	MenuItemTooltip = LOCTEXT("ObjectMixerEditorMenuItemTooltip", "Open an Object Mixer instance.");
+
+	// Should be hidden for now since it's not ready yet for public release
+	TabSpawnerType = ETabSpawnerMenuType::Hidden;
+}
+
+void FObjectMixerEditorModule::RegisterTabSpawner()
+{
+	FTabSpawnerEntry& BrowserSpawnerEntry =
+		FGlobalTabmanager::Get()->RegisterNomadTabSpawner(
+			GetTabSpawnerId(), 
+			FOnSpawnTab::CreateRaw(this, &FObjectMixerEditorModule::SpawnTab)
+		)
+		.SetIcon(MenuItemIcon)
+		.SetDisplayName(MenuItemName)
+		.SetTooltipText(MenuItemTooltip)
+		.SetMenuType(TabSpawnerType);
+
+	// Always use the base ObjectMixer function call or WorkspaceGroup may be null 
+	if (!FObjectMixerEditorModule::Get().RegisterItemInMenuGroup(BrowserSpawnerEntry))
+	{
+		BrowserSpawnerEntry.SetGroup(WorkspaceMenu::GetMenuStructure().GetLevelEditorCategory());
+	}
+}
+
+FName FObjectMixerEditorModule::GetTabSpawnerId()
+{
+	return "ObjectMixerToolkit";
+}
+
+bool FObjectMixerEditorModule::RegisterItemInMenuGroup(FWorkspaceItem& InItem)
+{
+	if (WorkspaceGroup)
+	{
+		WorkspaceGroup->AddItem(MakeShareable(&InItem));
+		
+		return true;
+	}
+
+	return false;
+}
+
+void FObjectMixerEditorModule::UnregisterTabSpawner()
+{
+	FGlobalTabmanager::Get()->UnregisterNomadTabSpawner(GetTabSpawnerId());
+}
+
+void FObjectMixerEditorModule::RegisterProjectSettings() const
+{
+	if (ISettingsModule* SettingsModule = FModuleManager::GetModulePtr<ISettingsModule>("Settings"))
+	{
+		// User Project Settings
+		const TSharedPtr<ISettingsSection> ProjectSettingsSectionPtr = SettingsModule->RegisterSettings(
+			"Project", "Editor", "Object Mixer",
+			LOCTEXT("ObjectMixerSettingsDisplayName", "Object Mixer"),
+			LOCTEXT("ObjectMixerSettingsDescription", "Configure Object Mixer user settings"),
+			GetMutableDefault<UObjectMixerEditorProjectSettings>());	
+	}
+}
+
+void FObjectMixerEditorModule::UnregisterProjectSettings() const
+{
+	if (ISettingsModule* SettingsModule = FModuleManager::GetModulePtr<ISettingsModule>("Settings"))
+	{
+		SettingsModule->UnregisterSettings("Project", "Plugins", "Object Mixer");
+	}
+}
+
+TSharedRef<SDockTab> FObjectMixerEditorModule::SpawnTab(const FSpawnTabArgs& Args)
+{
+	 return SpawnMainPanelTab();
+}
+
+TSharedPtr<FWorkspaceItem> FObjectMixerEditorModule::GetWorkspaceGroup()
+{
+	return WorkspaceGroup;
+}
+
+TSharedRef<SDockTab> FObjectMixerEditorModule::SpawnMainPanelTab()
+{
+	MainPanel = MakeShared<FObjectMixerEditorMainPanel>(GetModuleName());
+	MainPanel->Init();
+	
+	const TSharedRef<SDockTab> DockTab =
+		SNew(SDockTab)
+		.Label(TabLabel)
+		.TabRole(ETabRole::NomadTab)
+	;
+	const TSharedPtr<SWidget> ObjectMixerDialog = MakeObjectMixerDialog();
+	DockTab->SetContent(ObjectMixerDialog ? ObjectMixerDialog.ToSharedRef() : SNullWidget::NullWidget);
+
+	if (DefaultFilterClass)
+	{
+		MainPanel->OnClassSelectionChanged(DefaultFilterClass);
+	}
+			
+	return DockTab;
+}
+
+void FObjectMixerEditorModule::BindDelegates()
+{
+	check(GEngine && GEditor);
+	DelegateHandles.Add(GEngine->OnLevelActorFolderChanged().AddLambda([this] (const AActor*, FName)
+	{
+		RequestRebuildList();
+	}));
+	DelegateHandles.Add(GEngine->OnLevelActorListChanged().AddLambda([this] ()
+	{
+		RequestRebuildList();
+	}));
+	DelegateHandles.Add(GEngine->OnActorFolderAdded().AddLambda([this] (UActorFolder*)
+	{
+		RequestRebuildList();
+	}));
+	DelegateHandles.Add(GEngine->OnLevelActorDetached().AddLambda([this] (AActor*, const AActor*)
+	{
+		RequestRebuildList();
+	}));
+	DelegateHandles.Add(GEngine->OnLevelActorAttached().AddLambda([this] (AActor*, const AActor*)
+	{
+		RequestRebuildList();
+	}));
+	DelegateHandles.Add(GEngine->OnActorFolderAdded().AddLambda([this] (UActorFolder*)
+	{
+		RequestRebuildList();
+	}));
+	DelegateHandles.Add(GEngine->OnActorFoldersUpdatedEvent().AddLambda([this] (ULevel*)
+	{
+		RequestRebuildList();
+	}));
+
+	DelegateHandles.Add(GEditor->OnBlueprintCompiled().AddLambda([this] ()
+	{
+		RequestRebuildList();
+	}));
+
+	DelegateHandles.Add(FEditorDelegates::MapChange.AddLambda([this](uint32 Index)
+	{
+		RequestRebuildList();
+	}));
+
+	DelegateHandles.Add(FEditorDelegates::PostUndoRedo.AddLambda([this]()
+	{
+		RequestRebuildList();
+		
+		// Because we can undo/redo collection adds and removes, we need to save the data after each
+		GetMutableDefault<UObjectMixerEditorSerializedData>()->SaveConfig();
+	}));
+
+	DelegateHandles.Add(FCoreUObjectDelegates::OnObjectTransacted.AddLambda([this](UObject*, const FTransactionObjectEvent& Event)
+	{
+		if (Event.GetEventType() == ETransactionObjectEventType::Finalized)
+		{
+			RequestRebuildList();
+		}
+	}));
+
+	FLevelEditorModule& LevelEditorModule = FModuleManager::LoadModuleChecked<FLevelEditorModule>("LevelEditor");
+	LevelEditorModule.OnComponentsEdited().AddLambda([this]()
+	{
+		RequestRebuildList();
+	});
+}
+
+#undef LOCTEXT_NAMESPACE
diff --git a/Engine/Plugins/Editor/ObjectMixer/ObjectMixer/Source/ObjectMixer/Private/ObjectMixerEditorSerializedData.cpp b/Engine/Plugins/Editor/ObjectMixer/ObjectMixer/Source/ObjectMixer/Private/ObjectMixerEditorSerializedData.cpp
new file mode 100644
--- /dev/null
+++ b/Engine/Plugins/Editor/ObjectMixer/ObjectMixer/Source/ObjectMixer/Private/ObjectMixerEditorSerializedData.cpp
@@ -0,0 +1,194 @@
+// Copyright Epic Games, Inc. All Rights Reserved.
+
+#include "ObjectMixerEditorSerializedData.h"
+
+#include "Algo/Find.h"
+#include "Algo/RemoveIf.h"
+#include "ScopedTransaction.h"
+#include "Algo/IndexOf.h"
+
+#define LOCTEXT_NAMESPACE "ObjectMixerEditor"
+
+FObjectMixerSerializationData* UObjectMixerEditorSerializedData::FindSerializationDataByFilterClassName(const FName& FilterClassName)
+{
+	if (FObjectMixerSerializationData* Match = Algo::FindByPredicate(SerializedData,
+		[FilterClassName](const FObjectMixerSerializationData& Comparator)
+		{
+			return Comparator.FilterClassName.IsEqual(FilterClassName);
+		}))
+	{
+		return Match;
+	}
+
+	return nullptr;
+}
+
+void UObjectMixerEditorSerializedData::AddObjectsToCollection(const FName& FilterClassName, const FName& CollectionName, const TSet<FSoftObjectPath>& ObjectsToAdd)
+{
+	FScopedTransaction AddObjectsToCollectionTransaction(LOCTEXT("AddObjectsToCollectionTransaction","Add Objects To Collection"));
+
+	Modify();
+	
+	FObjectMixerSerializationData* Data = FindSerializationDataByFilterClassName(FilterClassName);
+
+	if (!Data)
+	{
+		SerializedData.Add({FilterClassName});
+		Data = FindSerializationDataByFilterClassName(FilterClassName);
+	}
+	
+	if (Data)
+	{
+		if (FObjectMixerCollectionObjectSet* Match = Algo::FindByPredicate(Data->SerializedCollection,
+			[CollectionName](const FObjectMixerCollectionObjectSet& Comparator)
+			{
+				return Comparator.CollectionName.IsEqual(CollectionName);
+			}))
+		{
+			(*Match).CollectionObjects.Append(ObjectsToAdd);
+		}
+		else
+		{
+			Data->SerializedCollection.Add({CollectionName, ObjectsToAdd});
+		}
+
+		SaveConfig();
+	}
+}
+
+void UObjectMixerEditorSerializedData::RemoveObjectsFromCollection(const FName& FilterClassName, const FName& CollectionName, const TSet<FSoftObjectPath>& ObjectsToRemove)
+{
+	if (FObjectMixerSerializationData* Data = FindSerializationDataByFilterClassName(FilterClassName))
+	{
+		FScopedTransaction RemoveCollectionTransaction(LOCTEXT("RemoveCollectionTransaction","Remove Collection"));
+	
+		Modify();
+		
+		Data->SerializedCollection.SetNum(Algo::StableRemoveIf(Data->SerializedCollection,
+			[CollectionName](const FObjectMixerCollectionObjectSet& Comparator)
+			{
+				return Comparator.CollectionName.IsEqual(CollectionName);
+			}));
+		
+		SaveConfig();
+	}
+}
+
+void UObjectMixerEditorSerializedData::RemoveCollection(const FName& FilterClassName, const FName& CollectionName)
+{	
+	if (FObjectMixerSerializationData* Data = FindSerializationDataByFilterClassName(FilterClassName))
+	{
+		FScopedTransaction RemoveCollectionTransaction(LOCTEXT("RemoveCollectionTransaction","Remove Collection"));
+	
+		Modify();
+		
+		Data->SerializedCollection.SetNum(Algo::StableRemoveIf(Data->SerializedCollection,
+			[CollectionName](const FObjectMixerCollectionObjectSet& Comparator)
+			{
+				return Comparator.CollectionName.IsEqual(CollectionName);
+			}));
+		
+		SaveConfig();
+	}
+}
+
+void UObjectMixerEditorSerializedData::ReorderCollection(const FName& FilterClassName,
+	const FName& CollectionToMoveName, const FName& CollectionInsertBeforeName)
+{
+	if (FObjectMixerSerializationData* Data = FindSerializationDataByFilterClassName(FilterClassName))
+	{
+		FObjectMixerCollectionObjectSet MatchCopy;
+		bool bFoundMatch = false;
+		if (const FObjectMixerCollectionObjectSet* Match = Algo::FindByPredicate(Data->SerializedCollection,
+			[CollectionToMoveName](const FObjectMixerCollectionObjectSet& Comparator)
+			{
+				return Comparator.CollectionName.IsEqual(CollectionToMoveName);
+			}))
+		{
+			MatchCopy = (*Match);
+			bFoundMatch = true;
+		}
+
+		if (bFoundMatch)
+		{
+			FScopedTransaction RemoveCollectionTransaction(LOCTEXT("RemoveCollectionTransaction","Remove Collection"));
+	
+			Modify();
+			
+			RemoveCollection(FilterClassName, CollectionToMoveName);
+
+			if (CollectionInsertBeforeName == "All") // Move MatchCopy to the end
+			{
+				Data->SerializedCollection.Add(MatchCopy);
+			}
+			else
+			{
+				const int32 IndexOfCollectionInsertBefore = Algo::IndexOfByPredicate(Data->SerializedCollection,
+					[CollectionInsertBeforeName](const FObjectMixerCollectionObjectSet& Comparator)
+					{
+						return Comparator.CollectionName.IsEqual(CollectionInsertBeforeName);
+					});
+
+				Data->SerializedCollection.Insert(MatchCopy, IndexOfCollectionInsertBefore);
+			}
+		
+			SaveConfig();
+		}
+	}
+}
+
+bool UObjectMixerEditorSerializedData::IsObjectInCollection(const FName& FilterClassName, const FName& CollectionName, const FSoftObjectPath& InObject)
+{
+	if (FObjectMixerSerializationData* Data = FindSerializationDataByFilterClassName(FilterClassName))
+	{
+		if (const FObjectMixerCollectionObjectSet* Match = Algo::FindByPredicate(Data->SerializedCollection,
+			[CollectionName](const FObjectMixerCollectionObjectSet& Comparator)
+			{
+				return Comparator.CollectionName.IsEqual(CollectionName);
+			}))
+		{
+			return (*Match).CollectionObjects.Contains(InObject);
+		}
+	}
+
+	return false;
+}
+
+TSet<FName> UObjectMixerEditorSerializedData::GetCollectionsForObject(const FName& FilterClassName, const FSoftObjectPath& InObject)
+{
+	if (FObjectMixerSerializationData* Data = FindSerializationDataByFilterClassName(FilterClassName))
+	{
+		TSet<FName> CollectionsWithObject;
+
+		for (const FObjectMixerCollectionObjectSet& CollectionObjectSet : Data->SerializedCollection)
+		{
+			if (IsObjectInCollection(FilterClassName, CollectionObjectSet.CollectionName, InObject))
+			{
+				CollectionsWithObject.Add(CollectionObjectSet.CollectionName);
+			}
+		}
+
+		return CollectionsWithObject;
+	}
+
+	return {};
+}
+
+TArray<FName> UObjectMixerEditorSerializedData::GetAllCollectionNames(const FName& FilterClassName)
+{
+	if (FObjectMixerSerializationData* Data = FindSerializationDataByFilterClassName(FilterClassName))
+	{
+		TArray<FName> Collections;
+
+		for (const FObjectMixerCollectionObjectSet& CollectionObjectSet : Data->SerializedCollection)
+		{
+			Collections.Add(CollectionObjectSet.CollectionName);
+		}
+
+		return Collections;
+	}
+
+	return {};
+}
+
+#undef LOCTEXT_NAMESPACE
diff --git a/Engine/Plugins/Editor/ObjectMixer/ObjectMixer/Source/ObjectMixer/Private/ObjectMixerEditorStyle.cpp b/Engine/Plugins/Editor/ObjectMixer/ObjectMixer/Source/ObjectMixer/Private/ObjectMixerEditorStyle.cpp
new file mode 100644
--- /dev/null
+++ b/Engine/Plugins/Editor/ObjectMixer/ObjectMixer/Source/ObjectMixer/Private/ObjectMixerEditorStyle.cpp
@@ -0,0 +1,103 @@
+// Copyright Epic Games, Inc. All Rights Reserved.
+
+#include "ObjectMixerEditorStyle.h"
+
+#include "Styling/AppStyle.h"
+#include "Framework/Application/SlateApplication.h"
+#include "Interfaces/IPluginManager.h"
+#include "Styling/CoreStyle.h"
+#include "Styling/SlateStyleRegistry.h"
+#include "Styling/SlateTypes.h"
+#include "Styling/StyleColors.h"
+
+TSharedPtr<FSlateStyleSet> FObjectMixerEditorStyle::StyleInstance = nullptr;
+
+void FObjectMixerEditorStyle::Initialize()
+{
+	if (!StyleInstance.IsValid())
+	{
+		StyleInstance = Create();
+		FSlateStyleRegistry::RegisterSlateStyle(*StyleInstance);
+	}
+}
+
+void FObjectMixerEditorStyle::Shutdown()
+{
+	FSlateStyleRegistry::UnRegisterSlateStyle(*StyleInstance);
+	ensure(StyleInstance.IsUnique());
+	StyleInstance.Reset();
+}
+
+void FObjectMixerEditorStyle::ReloadTextures()
+{
+	FSlateApplication::Get().GetRenderer()->ReloadTextureResources();
+}
+
+const ISlateStyle& FObjectMixerEditorStyle::Get()
+{
+	return *StyleInstance;
+}
+
+const FName& FObjectMixerEditorStyle::GetStyleSetName() const
+{
+	static FName StyleName(TEXT("ObjectMixerEditor"));
+	return StyleName;
+}
+
+FString GetConcertContent(const FString& RelativePath, const ANSICHAR* Extension)
+{
+	static FString ContentDir = IPluginManager::Get().FindPlugin(TEXT("ConcertSyncClient"))->GetContentDir();
+	return (ContentDir / RelativePath) + Extension;
+}
+
+const FSlateBrush* FObjectMixerEditorStyle::GetBrush(const FName PropertyName, const ANSICHAR* Specifier, const ISlateStyle* RequestingStyle) const
+{
+	return StyleInstance->GetBrush(PropertyName, Specifier);
+}
+
+#define IMAGE_BRUSH( RelativePath, ... ) FSlateImageBrush( Style->RootToContentDir( RelativePath, TEXT(".png") ), __VA_ARGS__ )
+#define BOX_BRUSH( RelativePath, ... ) FSlateBoxBrush( Style->RootToContentDir( RelativePath, TEXT(".png") ), __VA_ARGS__ )
+#define BORDER_BRUSH( RelativePath, ... ) FSlateBorderBrush( Style->RootToContentDir( RelativePath, TEXT(".png") ), __VA_ARGS__ )
+#define IMAGE_BRUSH_SVG( RelativePath, ... ) FSlateVectorImageBrush( Style->RootToContentDir(RelativePath, TEXT(".svg") ), __VA_ARGS__)
+#define IMAGE_PLUGIN_BRUSH_SVG( PluginName, RelativePath, ... ) FSlateVectorImageBrush( FObjectMixerEditorStyle::GetExternalPluginContent(PluginName, RelativePath, ".svg"), __VA_ARGS__)
+
+const FVector2D Icon64x64(64.f, 64.f);
+const FVector2D Icon40x40(40.0f, 40.0f);
+const FVector2D Icon20x20(20.0f, 20.0f);
+const FVector2D Icon16x16(16.0f, 16.0f);
+const FVector2D Icon12x12(12.0f, 12.0f);
+const FVector2D Icon8x8(8.f, 8.f);
+
+FString FObjectMixerEditorStyle::GetExternalPluginContent(const FString& PluginName, const FString& RelativePath, const ANSICHAR* Extension)
+{
+	FString ContentDir = IPluginManager::Get().FindPlugin(PluginName)->GetBaseDir() / RelativePath + Extension;
+	return ContentDir;
+}
+
+TSharedRef< FSlateStyleSet > FObjectMixerEditorStyle::Create()
+{
+	TSharedRef<FSlateStyleSet> Style = MakeShared<FSlateStyleSet>("ObjectMixerEditor");
+	
+	//Style->SetContentRoot(FPaths::Combine(Plugin->GetBaseDir(), TEXT("Resources")));
+
+	//Style->SetCoreContentRoot(FPaths::EngineContentDir() / TEXT("Slate"));
+
+	// Icons
+	Style->Set("ObjectMixer.ToolbarButton", new IMAGE_BRUSH_SVG("Icons/ObjectMixer", Icon40x40));
+	Style->Set("ObjectMixer.ToolbarButton.Small", new IMAGE_BRUSH_SVG("Icons/ObjectMixer", Icon20x20));
+
+	// Brush
+	Style->Set("ObjectMixerEditor.BrightBorder", new FSlateColorBrush(FStyleColors::Header));
+
+	// Border colors for Results view
+	Style->Set("ObjectMixerEditor.HeaderRowBorder", new FSlateColorBrush(FStyleColors::Black));
+	Style->Set("ObjectMixerEditor.DefaultBorder", new FSlateColorBrush(FStyleColors::Transparent));
+	
+	return Style;
+}
+
+#undef IMAGE_BRUSH
+#undef BOX_BRUSH
+#undef BORDER_BRUSH
+#undef IMAGE_BRUSH_SVG
+#undef IMAGE_PLUGIN_BRUSH_SVG
diff --git a/Engine/Plugins/Editor/ObjectMixer/ObjectMixer/Source/ObjectMixer/Private/ObjectMixerFilterFactory.cpp b/Engine/Plugins/Editor/ObjectMixer/ObjectMixer/Source/ObjectMixer/Private/ObjectMixerFilterFactory.cpp
new file mode 100644
--- /dev/null
+++ b/Engine/Plugins/Editor/ObjectMixer/ObjectMixer/Source/ObjectMixer/Private/ObjectMixerFilterFactory.cpp
@@ -0,0 +1,48 @@
+// Copyright Epic Games, Inc. All Rights Reserved.
+
+#include "ObjectMixerFilterFactory.h"
+
+#include "ObjectFilter/ObjectMixerEditorObjectFilter.h"
+
+#include "AssetToolsModule.h"
+#include "Kismet2/KismetEditorUtilities.h"
+
+#define LOCTEXT_NAMESPACE "ObjectMixerFilterFactory"
+
+UObjectMixerBlueprintFilterFactory::UObjectMixerBlueprintFilterFactory()
+{
+	SupportedClass = UBlueprint::StaticClass();
+	ParentClass = UObjectMixerBlueprintObjectFilter::StaticClass();
+
+	bCreateNew = true;
+	bEditAfterNew = true;
+}
+
+FText UObjectMixerBlueprintFilterFactory::GetDisplayName() const
+{
+	return LOCTEXT("DisplayName", "Object Mixer Filter");
+}
+
+UObject* UObjectMixerBlueprintFilterFactory::FactoryCreateNew(UClass* Class, UObject* InParent, FName Name, EObjectFlags Flags, UObject* Context, FFeedbackContext* Warn)
+{
+	UBlueprint* FilterBlueprint = nullptr;
+	if (ensure(SupportedClass == Class))
+	{
+		ensure(0 != (RF_Public & Flags));
+		FilterBlueprint = FKismetEditorUtilities::CreateBlueprint(ParentClass, InParent, Name, BPTYPE_Normal, UBlueprint::StaticClass(), UBlueprintGeneratedClass::StaticClass(), NAME_None);
+	}
+	return FilterBlueprint;	
+}
+
+FString UObjectMixerBlueprintFilterFactory::GetDefaultNewAssetName() const
+{
+	return "NewObjectMixerFilter";
+}
+
+uint32 UObjectMixerBlueprintFilterFactory::GetMenuCategories() const
+{
+	IAssetTools& AssetTools = FModuleManager::LoadModuleChecked<FAssetToolsModule>("AssetTools").Get();
+	return AssetTools.RegisterAdvancedAssetCategory("Object Mixer", LOCTEXT("AssetCategoryName", "Object Mixer"));
+}
+
+#undef LOCTEXT_NAMESPACE
diff --git a/Engine/Plugins/Editor/ObjectMixer/ObjectMixer/Source/ObjectMixer/Private/Views/List/ObjectMixerEditorList.cpp b/Engine/Plugins/Editor/ObjectMixer/ObjectMixer/Source/ObjectMixer/Private/Views/List/ObjectMixerEditorList.cpp
new file mode 100644
--- /dev/null
+++ b/Engine/Plugins/Editor/ObjectMixer/ObjectMixer/Source/ObjectMixer/Private/Views/List/ObjectMixerEditorList.cpp
@@ -0,0 +1,82 @@
+// Copyright Epic Games, Inc. All Rights Reserved.
+
+#include "Views/List/ObjectMixerEditorList.h"
+
+#include "Views/List/SObjectMixerEditorList.h"
+
+FObjectMixerEditorList::~FObjectMixerEditorList()
+{
+	FlushWidget();
+}
+
+void FObjectMixerEditorList::FlushWidget()
+{
+	ListWidget.Reset();
+}
+
+TSharedRef<SWidget> FObjectMixerEditorList::GetOrCreateWidget()
+{
+	if (!ListWidget.IsValid())
+	{
+		SAssignNew(ListWidget, SObjectMixerEditorList, SharedThis(this));
+	}
+
+	RequestRebuildList();
+
+	return ListWidget.ToSharedRef();
+}
+
+void FObjectMixerEditorList::ClearList() const
+{
+	if (ListWidget.IsValid())
+	{
+		ListWidget->ClearList();
+	}
+}
+
+void FObjectMixerEditorList::RequestRebuildList() const
+{
+	if (ListWidget.IsValid())
+	{
+		ListWidget->RequestRebuildList();
+	}
+}
+
+void FObjectMixerEditorList::RefreshList() const
+{
+	if (ListWidget.IsValid())
+	{
+		ListWidget->RefreshList();
+	}
+}
+
+void FObjectMixerEditorList::ExecuteListViewSearchOnAllRows(const FString& SearchString,
+	const bool bShouldRefreshAfterward)
+{
+	ListWidget->ExecuteListViewSearchOnAllRows(SearchString, bShouldRefreshAfterward);
+}
+
+void FObjectMixerEditorList::EvaluateIfRowsPassFilters(const bool bShouldRefreshAfterward) const
+{
+	ListWidget->EvaluateIfRowsPassFilters();
+}
+
+TWeakPtr<FObjectMixerEditorMainPanel> FObjectMixerEditorList::GetMainPanelModel()
+{
+	return MainPanelModelPtr;
+}
+
+TWeakPtr<FObjectMixerEditorListRow> FObjectMixerEditorList::GetSoloRow()
+{
+	return GetMainPanelModel().Pin()->GetSoloRow();
+}
+
+void FObjectMixerEditorList::SetSoloRow(TSharedRef<FObjectMixerEditorListRow> InRow)
+{
+	GetMainPanelModel().Pin()->SetSoloRow(InRow);
+}
+
+void FObjectMixerEditorList::ClearSoloRow()
+{
+	GetMainPanelModel().Pin()->ClearSoloRow();
+}
diff --git a/Engine/Plugins/Editor/ObjectMixer/ObjectMixer/Source/ObjectMixer/Private/Views/List/ObjectMixerEditorListRow.cpp b/Engine/Plugins/Editor/ObjectMixer/ObjectMixer/Source/ObjectMixer/Private/Views/List/ObjectMixerEditorListRow.cpp
new file mode 100644
--- /dev/null
+++ b/Engine/Plugins/Editor/ObjectMixer/ObjectMixer/Source/ObjectMixer/Private/Views/List/ObjectMixerEditorListRow.cpp
@@ -0,0 +1,427 @@
+// Copyright Epic Games, Inc. All Rights Reserved.
+
+#include "Views/List/ObjectMixerEditorListRow.h"
+
+#include "Algo/AllOf.h"
+#include "ClassIconFinder.h"
+#include "GameFramework/Actor.h"
+#include "Styling/SlateIconFinder.h"
+#include "Views/List/ObjectMixerEditorList.h"
+#include "Views/List/SObjectMixerEditorList.h"
+
+TSharedRef<FObjectMixerListRowDragDropOp> FObjectMixerListRowDragDropOp::New(const TArray<FObjectMixerEditorListRowPtr>& InItems)
+{
+	check(InItems.Num() > 0);
+
+	TSharedRef<FObjectMixerListRowDragDropOp> Operation = MakeShareable(
+		new FObjectMixerListRowDragDropOp());
+
+	Operation->DraggedItems = InItems;
+
+	Operation->DefaultHoverIcon = FAppStyle::Get().GetBrush("Graph.ConnectorFeedback.Error");
+
+	Operation->DefaultHoverText = NSLOCTEXT("ObjectMixerEditor","DefaultRowHoverText","Drop onto another row or Collection Button.");
+
+	Operation->Construct();
+
+	return Operation;
+}
+
+FObjectMixerEditorListRow::~FObjectMixerEditorListRow()
+{
+	FlushReferences();
+}
+
+void FObjectMixerEditorListRow::FlushReferences()
+{
+	if (ChildRows.Num())
+	{
+		ChildRows.Empty();
+	}
+}
+
+UObjectMixerObjectFilter* FObjectMixerEditorListRow::GetObjectFilter() const
+{
+	if (const TSharedPtr<SObjectMixerEditorList> PinnedListView = GetListViewPtr().Pin())
+	{
+		if (const TSharedPtr<FObjectMixerEditorList> PinnedListModel = PinnedListView->GetListModelPtr().Pin())
+		{
+			if (const TSharedPtr<FObjectMixerEditorMainPanel> PinnedMainPanel = PinnedListModel->GetMainPanelModel().Pin())
+			{
+				if (UObjectMixerObjectFilter* Filter = PinnedMainPanel->GetObjectFilter())
+				{
+					return Filter;
+				}
+			}
+		}
+	}
+
+	return nullptr;
+}
+
+bool FObjectMixerEditorListRow::IsObjectRefInSelectedCollections() const
+{
+	check (ListViewPtr.IsValid());
+	
+	if (RowType != EObjectMixerEditorListRowType::None && RowType != EObjectMixerEditorListRowType::Folder)
+	{
+		if (const TSharedPtr<FObjectMixerEditorList> ListModel = ListViewPtr.Pin()->GetListModelPtr().Pin())
+		{
+			if (const TSharedPtr<FObjectMixerEditorMainPanel> MainPanel = ListModel->GetMainPanelModel().Pin())
+			{
+				const TSet<FName>& CollectionSelection = MainPanel->GetCurrentCollectionSelection();
+
+				if (CollectionSelection.Num() == 0)
+				{
+					return true;
+				}
+
+				const TSet<FName>& ObjectAssignedCollections = MainPanel->GetCollectionsForObject(GetObject());
+
+				const TSet<FName> Intersection = CollectionSelection.Intersect(ObjectAssignedCollections);
+
+				return Intersection.Num() > 0;
+			}
+		}
+	}
+	return false;
+}
+
+FObjectMixerEditorListRow::EObjectMixerEditorListRowType FObjectMixerEditorListRow::GetRowType() const
+{
+	return RowType;
+}
+
+int32 FObjectMixerEditorListRow::GetSortOrder() const
+{
+	return SortOrder;
+}
+
+void FObjectMixerEditorListRow::SetSortOrder(const int32 InNewOrder)
+{
+	SortOrder = InNewOrder;
+}
+
+TWeakPtr<FObjectMixerEditorListRow> FObjectMixerEditorListRow::GetDirectParentRow() const
+{
+	return DirectParentRow;
+}
+
+void FObjectMixerEditorListRow::SetDirectParentRow(
+	const TWeakPtr<FObjectMixerEditorListRow>& InDirectParentRow)
+{
+	DirectParentRow = InDirectParentRow;
+}
+
+const TArray<FObjectMixerEditorListRowPtr>& FObjectMixerEditorListRow::GetChildRows() const
+{
+	return ChildRows;
+}
+
+int32 FObjectMixerEditorListRow::GetChildCount() const
+{
+	return ChildRows.Num();
+}
+
+void FObjectMixerEditorListRow::SetChildRows(const TArray<FObjectMixerEditorListRowPtr>& InChildRows)
+{
+	ChildRows = InChildRows;
+}
+
+void FObjectMixerEditorListRow::AddToChildRows(const FObjectMixerEditorListRowPtr& InRow)
+{
+	InRow->SetDirectParentRow(SharedThis(this));
+	ChildRows.AddUnique(InRow);
+	ChildRows.StableSort(SObjectMixerEditorList::SortByTypeThenName);
+}
+
+void FObjectMixerEditorListRow::InsertChildRowAtIndex(const FObjectMixerEditorListRowPtr& InRow,
+                                                           const int32 AtIndex)
+{
+	ChildRows.Insert(InRow, AtIndex);
+}
+
+bool FObjectMixerEditorListRow::GetIsTreeViewItemExpanded()
+{
+	return GetListViewPtr().Pin()->IsTreeViewItemExpanded(GetAsShared());
+}
+
+void FObjectMixerEditorListRow::SetIsTreeViewItemExpanded(const bool bNewExpanded)
+{
+	GetListViewPtr().Pin()->SetTreeViewItemExpanded(GetAsShared(), bNewExpanded);
+}
+
+bool FObjectMixerEditorListRow::GetShouldExpandAllChildren() const
+{
+	return bShouldExpandAllChildren;
+}
+
+void FObjectMixerEditorListRow::SetShouldExpandAllChildren(const bool bNewShouldExpandAllChildren)
+{
+	bShouldExpandAllChildren = bNewShouldExpandAllChildren;
+}
+
+bool FObjectMixerEditorListRow::MatchSearchTokensToSearchTerms(
+	const TArray<FString> InTokens, ESearchCase::Type InSearchCase)
+{
+	// If the search is cleared we'll consider the row to pass search
+	bool bMatchFound = true;
+	
+	if (CachedSearchTerms.IsEmpty())
+	{
+		CachedSearchTerms = GetDisplayNameOverride().ToString() + " ";
+
+		if (const TObjectPtr<UObject> Object = GetObject())
+		{			
+			if (const UObjectMixerObjectFilter* Filter = GetObjectFilter())
+			{
+				CachedSearchTerms += Filter->GetRowDisplayName(Object, GetTreeViewMode()).ToString();
+			}
+		}
+	}
+
+	// Match any
+	for (const FString& Token : InTokens)
+	{
+		// Match all of these
+		const FString SpaceDelimiter = " ";
+		TArray<FString> OutSpacedArray;
+		if (Token.Contains(SpaceDelimiter) && Token.ParseIntoArray(OutSpacedArray, *SpaceDelimiter, true) > 1)
+		{
+			bMatchFound = Algo::AllOf(OutSpacedArray, [this, InSearchCase](const FString& Comparator)
+			{
+				return CachedSearchTerms.Contains(Comparator, InSearchCase);
+			});
+		}
+		else
+		{
+			bMatchFound = CachedSearchTerms.Contains(Token, InSearchCase);
+		}
+
+		if (bMatchFound)
+		{
+			break;
+		}
+	}
+
+	bDoesRowMatchSearchTerms = bMatchFound;
+
+	return bMatchFound;
+}
+
+void FObjectMixerEditorListRow::ExecuteSearchOnChildNodes(const FString& SearchString) const
+{
+	TArray<FString> Tokens;
+
+	SearchString.ParseIntoArray(Tokens, TEXT(" "), true);
+
+	ExecuteSearchOnChildNodes(Tokens);
+}
+
+void FObjectMixerEditorListRow::ExecuteSearchOnChildNodes(const TArray<FString>& Tokens) const
+{
+	for (const FObjectMixerEditorListRowPtr& ChildRow : GetChildRows())
+	{
+		if (!ensure(ChildRow.IsValid()))
+		{
+			continue;
+		}
+
+		const bool bMatch = ChildRow->MatchSearchTokensToSearchTerms(Tokens);;
+
+		if (ChildRow->GetChildCount() > 0)
+		{
+			if (bMatch)
+			{
+				// If the group name matches then we pass an empty string to search child nodes since we want them all to be visible
+				ChildRow->ExecuteSearchOnChildNodes("");
+			}
+			else
+			{
+				// Otherwise we iterate over all child nodes to determine which should and should not be visible
+				ChildRow->ExecuteSearchOnChildNodes(Tokens);
+			}
+		}
+	}
+}
+
+bool FObjectMixerEditorListRow::GetDoesRowPassFilters() const
+{
+	return bDoesRowPassFilters;
+}
+
+void FObjectMixerEditorListRow::SetDoesRowPassFilters(const bool bPass)
+{
+	bDoesRowPassFilters = bPass;
+}
+
+bool FObjectMixerEditorListRow::GetIsSelected()
+{
+	check (ListViewPtr.IsValid());
+
+	return ListViewPtr.Pin()->IsTreeViewItemSelected(SharedThis(this));
+}
+
+bool FObjectMixerEditorListRow::ShouldRowWidgetBeVisible() const
+{
+	return (bDoesRowMatchSearchTerms && bDoesRowPassFilters) || HasVisibleChildRowWidgets();
+}
+
+EVisibility FObjectMixerEditorListRow::GetDesiredRowWidgetVisibility() const
+{
+	return ShouldRowWidgetBeVisible() ? EVisibility::Visible : EVisibility::Collapsed;
+}
+
+bool FObjectMixerEditorListRow::HasVisibleChildRowWidgets() const
+{
+	if (ChildRows.Num() == 0)
+	{
+		return false;
+	}
+
+	for (const auto& ChildRow : ChildRows)
+	{
+		if (ChildRow->ShouldRowWidgetBeVisible())
+		{
+			return true;
+		}
+	}
+
+	return false;
+}
+
+FText FObjectMixerEditorListRow::GetDisplayName()
+{
+	if (!GetDisplayNameOverride().IsEmpty())
+	{
+		return GetDisplayNameOverride();
+	}
+
+	if (const UObjectMixerObjectFilter* Filter = GetObjectFilter())
+	{
+		if (const TObjectPtr<UObject> Object = GetObject())
+		{
+			return Filter->GetRowDisplayName(Object, GetTreeViewMode());
+		}
+	}
+
+	return FText::GetEmpty();
+}
+
+EObjectMixerTreeViewMode FObjectMixerEditorListRow::GetTreeViewMode()
+{
+	const TSharedPtr<SObjectMixerEditorList> PinnedListModel = GetListViewPtr().Pin();
+	check(PinnedListModel);
+	
+	return PinnedListModel->GetTreeViewMode();
+}
+
+TArray<FObjectMixerEditorListRowPtr> FObjectMixerEditorListRow::GetSelectedTreeViewItems() const
+{
+	return ListViewPtr.Pin()->GetSelectedTreeViewItems();
+}
+
+const FSlateBrush* FObjectMixerEditorListRow::GetObjectIconBrush()
+{
+	if (GetRowType() == EObjectMixerEditorListRowType::None)
+	{
+		return nullptr;
+	}
+
+	if (GetRowType() == EObjectMixerEditorListRowType::Folder)
+	{
+		if (GetIsTreeViewItemExpanded() && ChildRows.Num())
+		{
+			return FAppStyle::Get().GetBrush(TEXT("SceneOutliner.FolderOpen"));
+		}
+		
+		return FAppStyle::Get().GetBrush(TEXT("SceneOutliner.FolderClosed"));
+	}
+	
+	if (UObject* RowObject = GetObject())
+	{
+		auto GetIconForActor = [](AActor* InActor)
+		{
+			FName IconName = InActor->GetCustomIconName();
+			if (IconName == NAME_None)
+			{
+				IconName = InActor->GetClass()->GetFName();
+			}
+
+			return FClassIconFinder::FindIconForActor(InActor);
+		};
+
+		if (AActor* AsActor = Cast<AActor>(RowObject))
+		{
+			return GetIconForActor(AsActor);
+		}
+	
+		if (RowObject->IsA(UActorComponent::StaticClass()))
+		{		
+			return FSlateIconFinder::FindIconBrushForClass(RowObject->GetClass(), TEXT("SCS.Component"));
+		}
+	}
+
+	return nullptr;
+}
+
+bool FObjectMixerEditorListRow::GetObjectVisibility()
+{
+	if (GetRowType() == EObjectMixerEditorListRowType::Folder)
+	{
+		// If any child returns true, the folder returns true
+		for (const FObjectMixerEditorListRowPtr& Child : GetChildRows())
+		{
+			if (Child->GetObjectVisibility())
+			{
+				return true;
+			}
+		}
+
+		// The folder's visibility is only returned false if all children also return false
+		return false;
+	}
+	
+	if (const UObjectMixerObjectFilter* Filter = GetObjectFilter())
+	{
+		return Filter->GetRowEditorVisibility(GetObject(), GetTreeViewMode());
+	}
+
+	return false;
+}
+
+void FObjectMixerEditorListRow::SetObjectVisibility(const bool bNewIsVisible, const bool bIsRecursive)
+{
+	if (const UObjectMixerObjectFilter* Filter = GetObjectFilter())
+	{
+		Filter->OnSetRowEditorVisibility(GetObject(), bNewIsVisible, GetTreeViewMode());
+
+		if (bIsRecursive)
+		{
+			for (const FObjectMixerEditorListRowPtr& Child : GetChildRows())
+			{
+				Child->SetObjectVisibility(bNewIsVisible, true);
+			}
+		}
+	}
+}
+
+bool FObjectMixerEditorListRow::IsThisRowSolo() const
+{
+	return GetListViewPtr().Pin()->GetSoloRow().HasSameObject(this);
+}
+
+void FObjectMixerEditorListRow::SetThisAsSoloRow()
+{
+	GetListViewPtr().Pin()->SetSoloRow(SharedThis(this));
+}
+
+void FObjectMixerEditorListRow::ClearSoloRow()
+{
+	GetListViewPtr().Pin()->ClearSoloRow();
+}
+
+FObjectMixerEditorListRowPtr FObjectMixerEditorListRow::GetAsShared()
+{
+	return SharedThis(this);
+}
diff --git a/Engine/Plugins/Editor/ObjectMixer/ObjectMixer/Source/ObjectMixer/Private/Views/List/SObjectMixerEditorList.cpp b/Engine/Plugins/Editor/ObjectMixer/ObjectMixer/Source/ObjectMixer/Private/Views/List/SObjectMixerEditorList.cpp
new file mode 100644
--- /dev/null
+++ b/Engine/Plugins/Editor/ObjectMixer/ObjectMixer/Source/ObjectMixer/Private/Views/List/SObjectMixerEditorList.cpp
@@ -0,0 +1,1234 @@
+// Copyright Epic Games, Inc. All Rights Reserved.
+
+#include "Views/List/SObjectMixerEditorList.h"
+
+// View Filters
+#include "Views/List/ObjectMixerEditorListFilters/IObjectMixerEditorListFilter.h"
+
+#include "ObjectFilter/ObjectMixerEditorObjectFilter.h"
+#include "ObjectMixerEditorLog.h"
+#include "ObjectMixerEditorProjectSettings.h"
+#include "ObjectMixerEditorStyle.h"
+#include "Views/List/SObjectMixerEditorListRow.h"
+
+#include "ActorBrowsingMode.h"
+#include "Algo/AllOf.h"
+#include "Algo/AnyOf.h"
+#include "Editor.h"
+#include "LevelEditor.h"
+#include "LevelEditorMenuContext.h"
+#include "LevelEditorViewport.h"
+#include "Selection.h"
+#include "Framework/MultiBox/MultiBoxBuilder.h"
+#include "DragAndDrop/AssetDragDropOp.h"
+#include "Editor/UnrealEdEngine.h"
+#include "PlacementMode/Public/IPlacementModeModule.h"
+#include "Styling/StyleColors.h"
+#include "UObject/UObjectIterator.h"
+#include "Views/Widgets/ObjectMixerEditorListMenuContext.h"
+#include "Widgets/Layout/SScrollBox.h"
+#include "Widgets/Layout/SWidgetSwitcher.h"
+#include "Widgets/Text/SRichTextBlock.h"
+#include "WorldPartition/WorldPartition.h"
+
+#define LOCTEXT_NAMESPACE "ObjectMixerEditorList"
+
+const FName SObjectMixerEditorList::ItemNameColumnName(TEXT("Builtin_Name"));
+const FName SObjectMixerEditorList::EditorVisibilityColumnName(TEXT("Builtin_EditorVisibility"));
+const FName SObjectMixerEditorList::EditorVisibilitySoloColumnName(TEXT("Builtin_EditorVisibilitySolo"));
+
+void SObjectMixerEditorList::Construct(const FArguments& InArgs, TSharedRef<FObjectMixerEditorList> ListModel)
+{
+	ListModelPtr = ListModel;
+	
+	// Set Default Sorting info
+	ActiveSortingType = EColumnSortMode::Ascending;
+	
+	HeaderRow = SNew(SHeaderRow)
+				.CanSelectGeneratedColumn(false)
+				.Visibility(EVisibility::Visible)
+				;
+	
+	ChildSlot
+	[
+		SNew(SVerticalBox)
+
+		+ SVerticalBox::Slot()
+		.HAlign(HAlign_Fill)
+		.VAlign(VAlign_Fill)
+		[
+			SNew(SWidgetSwitcher)
+			.WidgetIndex_Lambda([this]()
+			{
+				return DoesTreeViewHaveVisibleChildren() ? 0 : 1;
+			})
+			
+			+ SWidgetSwitcher::Slot()
+			.HAlign(HAlign_Fill)
+			.VAlign(VAlign_Fill)
+			.Padding(2.0f, 2.0f, 2.0f, 2.0f)
+			[
+				SAssignNew(TreeViewPtr, STreeView<FObjectMixerEditorListRowPtr>)
+				.HeaderRow(HeaderRow)
+				.SelectionMode(ESelectionMode::Multi)
+				.OnContextMenuOpening_Lambda([this]()
+				{
+					TSharedPtr<SWidget> Widget;
+					
+					if (GetListModelPtr().IsValid())
+					{
+						SelectedTreeItemsToSelectedInLevelEditor();
+						Widget = UObjectMixerEditorListMenuContext::CreateContextMenu(
+							{GetSelectedTreeViewItems(), GetListModelPtr().Pin()->GetMainPanelModel()});
+					}
+					
+					return Widget;
+				})
+				.OnSelectionChanged_Lambda([this] (const FObjectMixerEditorListRowPtr& Row, const ESelectInfo::Type SelectionType)
+				{
+					const bool bSyncSelectionEnabled = GetDefault<UObjectMixerEditorProjectSettings>()->bSyncSelection;
+					const bool bIsAltDown = FSlateApplication::Get().GetModifierKeys().IsAltDown();
+					const bool bShouldSyncSelection = bSyncSelectionEnabled ? !bIsAltDown : bIsAltDown;
+						
+					if (bShouldSyncSelection)
+					{
+						SelectedTreeItemsToSelectedInLevelEditor();
+					}
+				})
+				.OnMouseButtonDoubleClick_Lambda([this] (FObjectMixerEditorListRowPtr Row)
+				{							
+					if (Row->GetRowType() == FObjectMixerEditorListRow::MatchingObject ||
+						Row->GetRowType() == FObjectMixerEditorListRow::ContainerObject)
+					{
+						if (GCurrentLevelEditingViewportClient)
+						{
+							AActor* Actor = Cast<AActor>(Row->GetObject());
+
+							if (!Actor)
+							{
+								Actor = Row->GetObject()->GetTypedOuter<AActor>();
+							}
+
+							if (Actor)
+							{
+								FVector Origin;
+								FVector Extents;
+								Actor->GetActorBounds(false, Origin, Extents, true);
+								GCurrentLevelEditingViewportClient->FocusViewportOnBox(FBox(Origin - Extents, Origin + Extents));
+							}
+						}
+					}
+					else
+					{
+						Row->SetIsTreeViewItemExpanded(!Row->GetIsTreeViewItemExpanded());
+					}
+				})
+				.TreeItemsSource(&VisibleTreeViewObjects)
+				.OnGenerateRow_Lambda([this](FObjectMixerEditorListRowPtr Row, const TSharedRef<STableViewBase>& OwnerTable)
+					{
+						check(Row.IsValid());
+					
+						return SNew(SObjectMixerEditorListRow, TreeViewPtr.ToSharedRef(), Row)
+								.Visibility_Raw(Row.Get(), &FObjectMixerEditorListRow::GetDesiredRowWidgetVisibility);
+					})
+				.OnGetChildren_Raw(this, &SObjectMixerEditorList::OnGetRowChildren)
+				.OnExpansionChanged_Raw(this, &SObjectMixerEditorList::OnRowChildExpansionChange, false)
+				.OnSetExpansionRecursive(this, &SObjectMixerEditorList::OnRowChildExpansionChange, true)
+			]
+
+			// For when no rows exist in view
+			+ SWidgetSwitcher::Slot()
+			.HAlign(HAlign_Fill)
+			.Padding(2.0f, 24.0f, 2.0f, 2.0f)
+			[
+				SNew(SRichTextBlock)
+				.DecoratorStyleSet(&FAppStyle::Get())
+				.AutoWrapText(true)
+				.Justification(ETextJustify::Center)
+				.Text_Lambda([this]()
+				{
+					// Preset Empty List (with filter)
+                    return LOCTEXT("EmptyListPresetWithFilter", "No matching items in your list.\n\nCheck your filters.");
+				})
+			]
+		]
+	];
+}
+
+SObjectMixerEditorList::~SObjectMixerEditorList()
+{	
+	HeaderRow.Reset();
+
+	FlushMemory(false);
+
+	TreeViewPtr.Reset();
+}
+
+void SObjectMixerEditorList::RebuildList()
+{
+	bShouldRebuild = false;
+	
+	GenerateTreeView();
+}
+
+void SObjectMixerEditorList::RefreshList()
+{
+	if (TreeViewRootObjects.Num() > 0)
+	{
+		// Apply last search
+		ExecuteListViewSearchOnAllRows(GetSearchStringFromSearchInputField(), false);
+
+		// Enforce Sort
+		const FName& SortingName = GetActiveSortingColumnName();
+		ExecuteSort(SortingName, GetSortModeForColumn(SortingName), false);
+
+		// Show/Hide rows based on SetBy changes and filter settings
+		EvaluateIfRowsPassFilters(false);
+	}
+
+	FindVisibleObjectsAndRequestTreeRefresh();
+}
+
+void SObjectMixerEditorList::RequestRebuildList(const FString& InItemToScrollTo)
+{
+	bShouldRebuild = true;
+}
+
+TArray<FObjectMixerEditorListRowPtr> SObjectMixerEditorList::GetSelectedTreeViewItems() const
+{
+	return TreeViewPtr->GetSelectedItems();
+}
+
+int32 SObjectMixerEditorList::GetSelectedTreeViewItemCount() const
+{
+	return TreeViewPtr->GetSelectedItems().Num();
+}
+
+void SObjectMixerEditorList::SetSelectedTreeViewItemActorsEditorVisible(const bool bNewIsVisible, const bool bIsRecursive)
+{
+	for (const TSharedPtr<FObjectMixerEditorListRow>& SelectedItem : TreeViewPtr->GetSelectedItems())
+	{
+		SelectedItem->SetObjectVisibility(bNewIsVisible, bIsRecursive);
+	}
+}
+
+bool SObjectMixerEditorList::IsTreeViewItemSelected(TSharedRef<FObjectMixerEditorListRow> Item)
+{
+	return TreeViewPtr->GetSelectedItems().Contains(Item);
+}
+
+TArray<FObjectMixerEditorListRowPtr> SObjectMixerEditorList::GetTreeViewItems() const
+{
+	return TreeViewRootObjects;
+}
+
+void SObjectMixerEditorList::SetTreeViewItems(const TArray<FObjectMixerEditorListRowPtr>& InItems)
+{
+	TreeViewRootObjects = InItems;
+
+	TreeViewPtr->RequestListRefresh();
+}
+
+TWeakPtr<FObjectMixerEditorListRow> SObjectMixerEditorList::GetSoloRow()
+{
+	return GetListModelPtr().Pin()->GetSoloRow();
+}
+
+void SObjectMixerEditorList::SetSoloRow(TSharedRef<FObjectMixerEditorListRow> InRow)
+{
+	GetListModelPtr().Pin()->SetSoloRow(InRow);
+}
+
+void SObjectMixerEditorList::ClearSoloRow()
+{
+	GetListModelPtr().Pin()->ClearSoloRow();
+}
+
+FString SObjectMixerEditorList::GetSearchStringFromSearchInputField() const
+{
+	return ListModelPtr.Pin()->GetMainPanelModel().Pin()->GetSearchStringFromSearchInputField();
+}
+
+void SObjectMixerEditorList::ExecuteListViewSearchOnAllRows(
+	const FString& SearchString, const bool bShouldRefreshAfterward)
+{
+	TArray<FString> Tokens;
+	
+	// unquoted search equivalent to a match-any-of search
+	SearchString.ParseIntoArray(Tokens, TEXT("|"), true);
+	
+	for (const TSharedPtr<FObjectMixerEditorListRow>& ChildRow : TreeViewRootObjects)
+	{
+		if (!ensure(ChildRow.IsValid()))
+		{
+			continue;
+		}
+		
+		const bool bGroupMatch = ChildRow->MatchSearchTokensToSearchTerms(Tokens);
+		
+		// If the group name matches then we pass in an empty string so all child nodes are visible.
+		// If the name doesn't match, then we need to evaluate each child.
+		ChildRow->ExecuteSearchOnChildNodes(bGroupMatch ? "" : SearchString);
+	}
+
+	if (bShouldRefreshAfterward)
+	{
+		FindVisibleObjectsAndRequestTreeRefresh();
+	}
+}
+
+bool SObjectMixerEditorList::DoesTreeViewHaveVisibleChildren() const
+{
+	if (TreeViewPtr.IsValid())
+	{
+		for (const TSharedPtr<FObjectMixerEditorListRow>& Header : TreeViewRootObjects)
+		{
+			const EVisibility HeaderVisibility = Header->GetDesiredRowWidgetVisibility();
+			
+			if (HeaderVisibility != EVisibility::Hidden && HeaderVisibility != EVisibility::Collapsed)
+			{
+				return true;
+			}
+		}
+	}
+	
+	return false;
+}
+
+bool SObjectMixerEditorList::IsTreeViewItemExpanded(const TSharedPtr<FObjectMixerEditorListRow>& Row) const
+{
+	if (TreeViewPtr.IsValid())
+	{
+		return TreeViewPtr->IsItemExpanded(Row);
+	}
+	
+	return false;
+}
+
+void SObjectMixerEditorList::SetTreeViewItemExpanded(const TSharedPtr<FObjectMixerEditorListRow>& RowToExpand, const bool bNewExpansion) const
+{
+	if (TreeViewPtr.IsValid())
+	{
+		TreeViewPtr->SetItemExpansion(RowToExpand, bNewExpansion);
+	}
+}
+
+EObjectMixerTreeViewMode SObjectMixerEditorList::GetTreeViewMode()
+{
+	const TSharedPtr<FObjectMixerEditorList> PinnedListModel = GetListModelPtr().Pin();
+	check(PinnedListModel);
+
+	const TSharedPtr<FObjectMixerEditorMainPanel> PinnedMainPanel = PinnedListModel->GetMainPanelModel().Pin();
+	check(PinnedMainPanel);
+	
+	return PinnedMainPanel->GetTreeViewMode();
+}
+
+const TArray<TSharedRef<IObjectMixerEditorListFilter>>& SObjectMixerEditorList::GetShowFilters()
+{
+	const TSharedPtr<FObjectMixerEditorList> PinnedListModel = GetListModelPtr().Pin();
+	check(PinnedListModel);
+
+	const TSharedPtr<FObjectMixerEditorMainPanel> PinnedMainPanel = PinnedListModel->GetMainPanelModel().Pin();
+	check(PinnedMainPanel);
+	
+	return PinnedMainPanel->GetShowFilters();
+}
+
+void SObjectMixerEditorList::EvaluateIfRowsPassFilters(const bool bShouldRefreshAfterward)
+{
+	// Separate filters by type
+	
+	TSet<TSharedRef<IObjectMixerEditorListFilter>> MatchAnyOfFilters;
+	TSet<TSharedRef<IObjectMixerEditorListFilter>> MatchAllOfFilters;
+
+	for (const TSharedRef<IObjectMixerEditorListFilter>& Filter : GetShowFilters())
+	{
+		if (Filter->GetFilterMatchType() ==
+			IObjectMixerEditorListFilter::EObjectMixerEditorListFilterMatchType::MatchAll)
+		{
+			MatchAllOfFilters.Add(Filter);
+		}
+		else
+		{
+			MatchAnyOfFilters.Add(Filter);
+		}
+	}
+
+	const TFunction<void(const TArray<FObjectMixerEditorListRowPtr>&)> RecursiveRowIterator =
+		[&MatchAnyOfFilters, &MatchAllOfFilters, &RecursiveRowIterator](const TArray<FObjectMixerEditorListRowPtr>& InObjects)
+	{
+		for (const FObjectMixerEditorListRowPtr& Row : InObjects)
+		{
+			if (Row.IsValid())
+			{
+				auto FilterProjection = [&Row](const TSharedRef<IObjectMixerEditorListFilter>& Filter)
+				{
+					return Filter->GetIsFilterActive() ? Filter->DoesItemPassFilter(Row) : true;
+				};
+			
+				const bool bPassesAnyOf = MatchAnyOfFilters.Num() ? Algo::AnyOf(MatchAnyOfFilters, FilterProjection) : true;
+				const bool bPassesAllOf = MatchAllOfFilters.Num() ? Algo::AllOf(MatchAllOfFilters, FilterProjection) : true;
+		
+				Row->SetDoesRowPassFilters(bPassesAnyOf && bPassesAllOf);
+				
+				if (Row->GetChildCount() > 0)
+				{
+					RecursiveRowIterator(Row->GetChildRows());
+				}
+			}
+		}
+	};
+
+	RecursiveRowIterator(TreeViewRootObjects);
+
+	if (bShouldRefreshAfterward)
+	{
+		FindVisibleObjectsAndRequestTreeRefresh();
+	}
+}
+
+EColumnSortMode::Type SObjectMixerEditorList::GetSortModeForColumn(FName InColumnName) const
+{
+	EColumnSortMode::Type ColumnSortMode = EColumnSortMode::None;
+
+	if (GetActiveSortingColumnName().IsEqual(InColumnName))
+	{
+		ColumnSortMode = ActiveSortingType;
+	}
+
+	return ColumnSortMode;
+}
+
+void SObjectMixerEditorList::OnSortColumnCalled(EColumnSortPriority::Type Priority, const FName& ColumnName, EColumnSortMode::Type SortMode)
+{
+	ExecuteSort(ColumnName, CycleSortMode(ColumnName));
+}
+
+EColumnSortMode::Type SObjectMixerEditorList::CycleSortMode(const FName& InColumnName)
+{
+	const EColumnSortMode::Type PreviousColumnSortMode = GetSortModeForColumn(InColumnName);
+	ActiveSortingType = PreviousColumnSortMode ==
+		EColumnSortMode::Ascending ? EColumnSortMode::Descending : EColumnSortMode::Ascending;
+
+	ActiveSortingColumnName = InColumnName;
+	return ActiveSortingType;
+}
+
+void SObjectMixerEditorList::ExecuteSort(
+	const FName& InColumnName, const EColumnSortMode::Type InColumnSortMode, const bool bShouldRefreshAfterward)
+{	
+	if (bShouldRefreshAfterward)
+	{
+		FindVisibleObjectsAndRequestTreeRefresh();
+	}
+}
+
+FListViewColumnInfo* SObjectMixerEditorList::GetColumnInfoByPropertyName(const FName& InPropertyName)
+{
+	return Algo::FindByPredicate(ListViewColumns,
+		[InPropertyName] (const FListViewColumnInfo& ColumnInfo)
+		{
+			return ColumnInfo.PropertyName.IsEqual(InPropertyName);
+		});
+}
+
+void SObjectMixerEditorList::Tick(const FGeometry& AllottedGeometry, const double InCurrentTime,
+	const float InDeltaTime)
+{
+	SCompoundWidget::Tick(AllottedGeometry, InCurrentTime, InDeltaTime);
+
+	if (bShouldRebuild)
+	{
+		RebuildList();
+	}
+}
+
+TSharedRef<SWidget> SObjectMixerEditorList::GenerateHeaderRowContextMenu() const
+{
+	FMenuBuilder MenuBuilder(false, nullptr);
+	
+	MenuBuilder.AddSearchWidget();
+
+	FName LastPropertyCategoryName = NAME_None;
+
+	for (const FListViewColumnInfo& ColumnInfo : ListViewColumns)
+	{
+		const FName& PropertyCategoryName = ColumnInfo.PropertyCategoryName;
+
+		if (!PropertyCategoryName.IsEqual(LastPropertyCategoryName))
+		{
+			LastPropertyCategoryName = PropertyCategoryName;
+			
+			MenuBuilder.EndSection();
+            MenuBuilder.BeginSection(LastPropertyCategoryName, FText::FromName(LastPropertyCategoryName));
+		}
+		
+		const FName& PropertyName = ColumnInfo.PropertyName;
+		
+		const FText Tooltip = ColumnInfo.PropertyRef ?
+			ColumnInfo.PropertyRef->GetToolTipText() : ColumnInfo.PropertyDisplayText;
+
+		const bool bCanSelectColumn = ColumnInfo.PropertyType != EListViewColumnType::BuiltIn;
+
+		const FName Hook = ColumnInfo.PropertyType == EListViewColumnType::BuiltIn ? "Builtin" : "GeneratedProperties";
+		
+		MenuBuilder.AddMenuEntry(
+			ColumnInfo.PropertyDisplayText,
+			Tooltip,
+			FSlateIcon(),
+			FUIAction(
+				FExecuteAction::CreateLambda([this, PropertyName]()
+				{
+					check(HeaderRow);
+					
+					HeaderRow->SetShowGeneratedColumn(PropertyName, !HeaderRow->IsColumnVisible(PropertyName));
+				}),
+				FCanExecuteAction::CreateLambda([bCanSelectColumn](){return bCanSelectColumn;}),
+				FIsActionChecked::CreateLambda([this, PropertyName]()
+				{
+					check(HeaderRow);
+					
+					return HeaderRow->IsColumnVisible(PropertyName);
+				})
+			),
+			Hook,
+			EUserInterfaceActionType::ToggleButton
+		);
+	}
+
+	return MenuBuilder.MakeWidget();
+}
+
+bool SObjectMixerEditorList::AddUniquePropertyColumnsToHeaderRow(
+	FProperty* Property,
+	const bool bForceIncludeProperty,
+	const TSet<FName>& PropertySkipList)
+{
+	if (!ensureAlwaysMsgf(Property, TEXT("%hs: Invalid property passed in. Please ensure only valid properties are passed to this function."), __FUNCTION__))
+	{
+		return false;
+	}
+
+	bool bShouldIncludeProperty = bForceIncludeProperty;
+
+	if (!bShouldIncludeProperty)
+	{
+		const bool bIsPropertyBlueprintEditable = (Property->GetPropertyFlags() & CPF_Edit) != 0;
+
+		// We don't have a proper way to display these yet
+		const bool bDoesPropertyHaveSupportedClass =
+			!Property->IsA(FMapProperty::StaticClass()) &&
+			!Property->IsA(FArrayProperty::StaticClass()) &&
+			!Property->IsA(FSetProperty::StaticClass()) &&
+			!Property->IsA(FStructProperty::StaticClass());
+
+		bShouldIncludeProperty = bIsPropertyBlueprintEditable && bDoesPropertyHaveSupportedClass;
+	}
+
+	if (bShouldIncludeProperty)
+	{
+		const bool bIsPropertyExplicitlySkipped =
+		   PropertySkipList.Num() && PropertySkipList.Contains(Property->GetFName());
+
+		bShouldIncludeProperty = !bIsPropertyExplicitlySkipped;
+	}
+	
+	if (bShouldIncludeProperty)
+	{
+		const FName PropertyName = Property->GetFName();
+	
+		// Ensure no duplicate properties
+		if (!Algo::FindByPredicate(ListViewColumns,
+				[&PropertyName] (const FListViewColumnInfo& ListViewColumn)
+				{
+					return ListViewColumn.PropertyName.IsEqual(PropertyName);
+				})
+			)
+		{
+			FString PropertyCategoryName = "Generated Properties";
+			if (const FString* CategoryMeta = Property->FindMetaData("Category"))
+			{
+				PropertyCategoryName = *CategoryMeta;
+			}
+			
+			ListViewColumns.Add(
+				{
+					Property, PropertyName,
+					Property->GetDisplayNameText(),
+					EListViewColumnType::PropertyGenerated,
+					*PropertyCategoryName,
+					true, false, false
+				}
+			);
+
+			return true;
+		}
+	}
+
+	return false;
+}
+
+void SObjectMixerEditorList::AddBuiltinColumnsToHeaderRow()
+{
+	ListViewColumns.Insert(
+		{
+			nullptr, ItemNameColumnName,
+			LOCTEXT("ItemNameHeaderText", "Name"),
+			EListViewColumnType::BuiltIn, "Built-In",
+			true, false,
+			false, 1.0f, 1.7f
+		}, 0
+	);
+	
+	ListViewColumns.Insert(
+		{
+			nullptr, EditorVisibilitySoloColumnName,
+			LOCTEXT("EditorVisibilitySoloColumnNameHeaderText", "Solo"),
+			EListViewColumnType::BuiltIn, "Built-In",
+			true, false,
+			true, 25.0f
+		}, 0
+	);
+
+	ListViewColumns.Insert(
+		{
+			nullptr, EditorVisibilityColumnName,
+			LOCTEXT("EditorVisibilityColumnNameHeaderText", "Visibility"),
+			EListViewColumnType::BuiltIn, "Built-In",
+			true, false,
+			true, 25.0f
+		}, 0
+	);
+}
+
+TSharedPtr<SHeaderRow> SObjectMixerEditorList::GenerateHeaderRow()
+{
+	check(ListModelPtr.IsValid());
+	check(ListModelPtr.Pin()->GetMainPanelModel().Pin());
+	check(HeaderRow);
+	
+	TMap<FName, bool> LastVisibleColumns;
+	
+	for (const SHeaderRow::FColumn& Column : HeaderRow->GetColumns())
+	{
+		LastVisibleColumns.Add(Column.ColumnId, Column.bIsVisible);
+	}
+
+	HeaderRow->ClearColumns();
+	ListViewColumns.Empty(ListViewColumns.Num());
+	
+	TSet<UClass*> SpecifiedClasses =
+		ListModelPtr.Pin()->GetMainPanelModel().Pin()->GetObjectFilter()->GetParentAndChildClassesFromSpecifiedClasses(
+			ObjectClassesToFilterCache, PropertyInheritanceInclusionOptionsCache);
+	
+	for (const UClass* Class : SpecifiedClasses)
+	{
+		for (TFieldIterator<FProperty> FieldIterator(Class, EFieldIterationFlags::None); FieldIterator; ++FieldIterator)
+		{
+			if (FProperty* Property = *FieldIterator)
+			{
+				AddUniquePropertyColumnsToHeaderRow(Property, bShouldIncludeUnsupportedPropertiesCache, ColumnsToExcludeCache);
+			}
+		}
+
+		// Check Force Added Columns
+		for (const FName& PropertyName : ForceAddedColumnsCache)
+		{
+			if (FProperty* Property = FindFProperty<FProperty>(Class, PropertyName))
+			{
+				AddUniquePropertyColumnsToHeaderRow(Property, true);
+			}
+		}
+	}
+
+	// Alphabetical sort by Property Name
+	ListViewColumns.StableSort([](const FListViewColumnInfo& A, const FListViewColumnInfo& B)
+	{
+		return A.PropertyDisplayText.ToString() < B.PropertyDisplayText.ToString();
+	});
+
+	// Alphabetical sort by Property Category Name
+	ListViewColumns.StableSort([](const FListViewColumnInfo& A, const FListViewColumnInfo& B)
+	{
+		return A.PropertyCategoryName.LexicalLess(B.PropertyCategoryName);
+	});
+
+	// Add Built-in Columns to beginning
+	AddBuiltinColumnsToHeaderRow();
+
+	// Actually add columns to Header
+	{	
+		const FText ClickToSortTooltip = LOCTEXT("ClickToSort","Click to sort");
+
+		const TSharedRef<SWidget> HeaderMenuContent = GenerateHeaderRowContextMenu();
+
+		for (const FListViewColumnInfo& ColumnInfo : ListViewColumns)
+		{
+			const FText Tooltip = ColumnInfo.PropertyRef ? ColumnInfo.PropertyRef->GetToolTipText() :
+				ColumnInfo.bCanBeSorted ? ClickToSortTooltip : ColumnInfo.PropertyDisplayText;
+		
+			SHeaderRow::FColumn::FArguments Column =
+				SHeaderRow::Column(ColumnInfo.PropertyName)
+				.DefaultLabel(ColumnInfo.PropertyDisplayText)
+				.ToolTipText(Tooltip)
+				.HAlignHeader(EHorizontalAlignment::HAlign_Left)
+				//.bHideHeaderMenuButton(false) todo: uncomment this line when 21478974 is pushed 
+			;
+
+			if (ColumnInfo.bUseFixedWidth)
+			{
+				Column.FixedWidth(ColumnInfo.FixedWidth);
+			}
+			else
+			{
+				Column.FillWidth(ColumnInfo.FillWidth);
+			}
+
+			if (ColumnInfo.bCanBeSorted)
+			{
+				Column.SortMode_Raw(this, &SObjectMixerEditorList::GetSortModeForColumn, ColumnInfo.PropertyName);
+				Column.OnSort_Raw(this, &SObjectMixerEditorList::OnSortColumnCalled);
+			}
+
+			if (ColumnInfo.PropertyType == EListViewColumnType::BuiltIn)
+			{
+				Column.ShouldGenerateWidget(true);
+			}
+
+			if (ColumnInfo.PropertyName.IsEqual(EditorVisibilityColumnName))
+			{
+				Column.HeaderContent()
+				[
+					SNew(SBox)
+					.HAlign(HAlign_Left)
+					.VAlign(VAlign_Center)
+					.Padding(0.f)
+					[
+						SNew(SImage)
+						.ColorAndOpacity(FSlateColor::UseForeground())
+						.Image(FAppStyle::Get().GetBrush("Level.VisibleIcon16x"))
+						.ToolTipText(LOCTEXT("VisibilityColumnTooltip", "Visibility"))
+					]
+				];
+			}
+			else if (ColumnInfo.PropertyName.IsEqual(EditorVisibilitySoloColumnName))
+			{
+				Column.HeaderContent()
+				[
+					SNew(SBox)
+					.HAlign(HAlign_Left)
+					.VAlign(VAlign_Center)
+					.Padding(0.f)
+					[
+						SNew(SImage)
+						.ColorAndOpacity(FSlateColor::UseForeground())
+						.Image(FAppStyle::Get().GetBrush("MediaAsset.AssetActions.Solo.Small"))
+						.ToolTipText(LOCTEXT("SoloColumnTooltip", "Solo"))
+					]
+				];
+			}
+			else // Add Column Selection Menu widget to all other columns
+			{
+				Column.MenuContent()
+				[
+					HeaderMenuContent
+				];
+			}
+		
+			HeaderRow->AddColumn(Column);
+			bool bShouldShowColumn = ColumnsToShowByDefaultCache.Contains(ColumnInfo.PropertyName);
+
+			if (const bool* Match = LastVisibleColumns.Find(ColumnInfo.PropertyName))
+			{
+				bShouldShowColumn = *Match;
+			}
+		
+			HeaderRow->SetShowGeneratedColumn(ColumnInfo.PropertyName, bShouldShowColumn);
+		}
+	}
+
+	return HeaderRow;
+}
+
+void SObjectMixerEditorList::FlushMemory(const bool bShouldKeepMemoryAllocated)
+{
+	if (bShouldKeepMemoryAllocated)
+	{
+		TreeViewRootObjects.Reset();
+		VisibleTreeViewObjects.Reset();
+	}
+	else
+	{
+		TreeViewRootObjects.Empty();
+		VisibleTreeViewObjects.Empty();
+	}
+}
+
+void SObjectMixerEditorList::SetAllGroupsCollapsed()
+{
+	if (TreeViewPtr.IsValid())
+	{
+		for (const FObjectMixerEditorListRowPtr& RootRow : TreeViewRootObjects)
+		{
+			if (!RootRow.IsValid())
+			{
+				continue;
+			}
+			
+			TreeViewPtr->SetItemExpansion(RootRow, false);
+			RootRow->SetIsTreeViewItemExpanded(false);
+		}
+	}
+}
+
+void SObjectMixerEditorList::CacheTreeState()
+{
+	struct Local
+	{
+		static void RecursivelyCacheTreeState(
+			const TArray<FObjectMixerEditorListRowPtr>& InObjects,
+			TMap<FString, bool>& TreeItemExpansionStateCache,
+			TSharedPtr<STreeView<FObjectMixerEditorListRowPtr>> TreeViewPtr)
+		{
+			for (const TSharedPtr<FObjectMixerEditorListRow>& TreeViewItem : InObjects)
+			{
+				const FString ObjectName = TreeViewItem->GetDisplayName().ToString();
+
+				if (!ObjectName.IsEmpty())
+				{
+					TreeItemExpansionStateCache.Add(ObjectName, TreeViewPtr->IsItemExpanded(TreeViewItem));
+				}
+
+				RecursivelyCacheTreeState(TreeViewItem->GetChildRows(), TreeItemExpansionStateCache, TreeViewPtr);
+			}
+		}
+	};
+
+	Local::RecursivelyCacheTreeState(TreeViewRootObjects, TreeItemExpansionStateCache, TreeViewPtr);
+}
+
+void SObjectMixerEditorList::RestoreTreeState(const bool bFlushCache)
+{
+	struct Local
+	{
+		static void RecursivelyRestoreTreeState(
+			const TArray<FObjectMixerEditorListRowPtr>& InObjects,
+			TMap<FString, bool>& TreeItemExpansionStateCache,
+			TSharedPtr<STreeView<FObjectMixerEditorListRowPtr>> TreeViewPtr,
+			const bool bExpandByDefault)
+		{
+			for (const TSharedPtr<FObjectMixerEditorListRow>& TreeViewItem : InObjects)
+			{
+				const FString ObjectName = TreeViewItem->GetDisplayName().ToString();
+
+				if (!ObjectName.IsEmpty())
+				{
+					if (const bool* bExpansionStatePtr = TreeItemExpansionStateCache.Find(ObjectName))
+					{
+						TreeViewPtr->SetItemExpansion(TreeViewItem, *bExpansionStatePtr);
+					}
+					else
+					{
+						TreeViewPtr->SetItemExpansion(TreeViewItem, bExpandByDefault);
+					}
+				}
+
+				RecursivelyRestoreTreeState(
+					TreeViewItem->GetChildRows(), TreeItemExpansionStateCache, TreeViewPtr, bExpandByDefault);
+			}
+		}
+	};
+	
+	bool bExpandByDefault = true;
+	if (const UObjectMixerEditorProjectSettings* const Settings = GetDefault<UObjectMixerEditorProjectSettings>())
+	{
+		bExpandByDefault = Settings->bExpandTreeViewItemsByDefault;
+	}
+	
+	Local::RecursivelyRestoreTreeState(
+		TreeViewRootObjects, TreeItemExpansionStateCache, TreeViewPtr, bExpandByDefault);
+
+	if (bFlushCache)
+	{
+		TreeItemExpansionStateCache.Empty();
+	}
+}
+
+void SObjectMixerEditorList::BuildPerformanceCacheAndGenerateHeaderIfNeeded()
+{
+	check(ListModelPtr.IsValid());
+	check(ListModelPtr.Pin()->GetMainPanelModel().Pin());
+	
+	// If any of the following overrides change, we need to regenerate the header row. Otherwise skip regeneration for performance reasons.
+	// GetObjectClassesToFilter, GetColumnsToShowByDefault, GetColumnsToExclude,
+	// GetForceAddedColumns, GetObjectMixerPropertyInheritanceInclusionOptions, ShouldIncludeUnsupportedProperties
+	bool bNeedToGenerateHeaders = false;
+	
+	const TObjectPtr<UObjectMixerObjectFilter> SelectedFilter = ListModelPtr.Pin()->GetMainPanelModel().Pin()->GetObjectFilter();
+	if (!SelectedFilter)
+	{
+		UE_LOG(LogObjectMixerEditor, Display, TEXT("%hs: No classes defined in UObjectMixerObjectFilter class."), __FUNCTION__);
+		return;
+	}
+
+	if (const TSet<UClass*> ObjectClassesToFilter = SelectedFilter->GetObjectClassesToFilter();
+		ObjectClassesToFilter.Num() != ObjectClassesToFilterCache.Num() ||
+		ObjectClassesToFilter.Difference(ObjectClassesToFilterCache).Num() > 0 || ObjectClassesToFilterCache.Difference(ObjectClassesToFilter).Num() > 0)
+	{
+		ObjectClassesToFilterCache = ObjectClassesToFilter;
+		if (!bNeedToGenerateHeaders)
+		{
+			bNeedToGenerateHeaders = true;
+		}
+	}
+	if (const TSet<FName> ColumnsToShowByDefault = SelectedFilter->GetColumnsToShowByDefault();
+		ColumnsToShowByDefault.Num() != ColumnsToShowByDefaultCache.Num() ||
+		ColumnsToShowByDefault.Difference(ColumnsToShowByDefaultCache).Num() > 0 || ColumnsToShowByDefaultCache.Difference(ColumnsToShowByDefault).Num() > 0)
+	{
+		ColumnsToShowByDefaultCache = ColumnsToShowByDefault;
+		if (!bNeedToGenerateHeaders)
+		{
+			bNeedToGenerateHeaders = true;
+		}
+	}
+	if (const TSet<FName> ColumnsToExclude = SelectedFilter->GetColumnsToExclude();
+		ColumnsToExclude.Num() != ColumnsToExcludeCache.Num() ||
+		ColumnsToExclude.Difference(ColumnsToExcludeCache).Num() > 0 || ColumnsToExcludeCache.Difference(ColumnsToExclude).Num() > 0)
+	{
+		ColumnsToExcludeCache = ColumnsToExclude;
+		if (!bNeedToGenerateHeaders)
+		{
+			bNeedToGenerateHeaders = true;
+		}
+	}
+	if (const TSet<FName> ForceAddedColumns = SelectedFilter->GetForceAddedColumns();
+		ForceAddedColumns.Num() != ForceAddedColumnsCache.Num() ||
+		ForceAddedColumns.Difference(ForceAddedColumnsCache).Num() > 0 || ForceAddedColumnsCache.Difference(ForceAddedColumns).Num() > 0)
+	{
+		ForceAddedColumnsCache = ForceAddedColumns;
+		if (!bNeedToGenerateHeaders)
+		{
+			bNeedToGenerateHeaders = true;
+		}
+	}
+	if (const EObjectMixerInheritanceInclusionOptions PropertyInheritanceInclusionOptions = SelectedFilter->GetObjectMixerPropertyInheritanceInclusionOptions(); 
+		PropertyInheritanceInclusionOptions != PropertyInheritanceInclusionOptionsCache)
+	{
+		PropertyInheritanceInclusionOptionsCache = PropertyInheritanceInclusionOptions;
+		if (!bNeedToGenerateHeaders)
+		{
+			bNeedToGenerateHeaders = true;
+		}
+	}
+	if (const bool bShouldIncludeUnsupportedProperties = SelectedFilter->ShouldIncludeUnsupportedProperties(); 
+		bShouldIncludeUnsupportedProperties != bShouldIncludeUnsupportedPropertiesCache)
+	{
+		bShouldIncludeUnsupportedPropertiesCache = bShouldIncludeUnsupportedProperties;
+		if (!bNeedToGenerateHeaders)
+		{
+			bNeedToGenerateHeaders = true;
+		}
+	}
+	
+	if (bNeedToGenerateHeaders)
+	{
+		GenerateHeaderRow();
+	}
+}
+
+void SObjectMixerEditorList::GenerateTreeView()
+{
+	check(ListModelPtr.IsValid());
+	
+	if (!ensure(TreeViewPtr.IsValid()))
+	{
+		return;
+	}
+
+	CacheTreeState();
+	
+	FlushMemory(true);
+
+	BuildPerformanceCacheAndGenerateHeaderIfNeeded();
+
+	check(GEditor);
+	const UWorld* EditorWorld = GEditor->GetEditorWorldContext().World();
+
+	// Find valid matching objects
+	TArray<UObject*> MatchingObjects;
+	for (TObjectIterator<UObject> ObjectIterator; ObjectIterator; ++ObjectIterator)
+	{
+		if (UObject* Object = *ObjectIterator; IsValid(Object))
+		{
+			if (UWorld* ObjectWorld = Object->GetWorld(); ObjectWorld == EditorWorld)
+			{
+				bool bIsAcceptableClass = false;
+
+				for (UClass* Class : ObjectClassesToFilterCache)
+				{
+					if (Object->IsA(Class))
+					{
+						bIsAcceptableClass = true;
+						break;
+					}
+				}
+
+				if (bIsAcceptableClass)
+				{
+					MatchingObjects.Add(Object);
+				}
+			}
+		}
+	}
+
+	// A quick lookup for objects that already exist in the list.
+	// Helpful to avoid double-generating rows when considering parent->child hierarchy.
+	TMap<UObject*, TSharedRef<FObjectMixerEditorListRow>> CreatedObjectMap;
+	TMap<FName, TSharedRef<FObjectMixerEditorListRow>> FolderMap;
+	for (UObject* Object : MatchingObjects)
+	{
+		if (!CreatedObjectMap.Contains(Object)) // Ensure we don't double-generate container objects
+		{
+			TSharedRef<FObjectMixerEditorListRow> TopLevelRow = MakeShared<FObjectMixerEditorListRow>(
+				Object, FObjectMixerEditorListRow::MatchingObject, SharedThis(this));
+
+			CreatedObjectMap.Add(Object, TopLevelRow);
+
+			// If the view is not in flat mode, we need to consider the hierarchy of outliner folders/attach parents as desired
+			if (GetTreeViewMode() != EObjectMixerTreeViewMode::Flat)
+			{
+				AActor* BaseActor = Cast<AActor>(Object);
+				
+				if (!BaseActor)
+				{
+					if (Object->IsA(UActorComponent::StaticClass()))
+					{
+						BaseActor = Object->GetTypedOuter<AActor>();
+
+						// If it's not flat or folder view mode, we need to find or create the container object for the actor that owns the matching component
+						if (GetTreeViewMode() != EObjectMixerTreeViewMode::Folder)
+						{
+							TSharedPtr<FObjectMixerEditorListRow> OwningActorRow;
+
+							if (BaseActor)
+							{
+								if (const TSharedRef<FObjectMixerEditorListRow>* Match = CreatedObjectMap.Find(BaseActor))
+								{
+									OwningActorRow = *Match;
+								}
+								else
+								{
+									OwningActorRow = MakeShared<FObjectMixerEditorListRow>(
+										BaseActor, FObjectMixerEditorListRow::ContainerObject, SharedThis(this));
+
+									CreatedObjectMap.Add(BaseActor, OwningActorRow.ToSharedRef());
+								}
+
+								if (OwningActorRow)
+								{
+									OwningActorRow->AddToChildRows(TopLevelRow);
+													
+									TopLevelRow = OwningActorRow.ToSharedRef();
+								}
+							}
+						}
+					}
+				}
+
+				if (BaseActor)
+				{
+					while (AActor* AttachParent = BaseActor->GetAttachParentActor())
+					{
+						// Make a row for each attach parent up the chain until we reach the top if not in flat/folder mode
+						if (GetTreeViewMode() != EObjectMixerTreeViewMode::Folder)
+						{
+							TSharedPtr<FObjectMixerEditorListRow> OwningActorRow;
+
+							if (const TSharedRef<FObjectMixerEditorListRow>* Match = CreatedObjectMap.Find(AttachParent))
+							{
+								OwningActorRow = *Match;
+							}
+							else
+							{
+								OwningActorRow = MakeShared<FObjectMixerEditorListRow>(
+									AttachParent, FObjectMixerEditorListRow::ContainerObject, SharedThis(this));
+
+								CreatedObjectMap.Add(AttachParent, OwningActorRow.ToSharedRef());
+							}
+
+							if (OwningActorRow)
+							{
+								OwningActorRow->AddToChildRows(TopLevelRow);
+												
+								TopLevelRow = OwningActorRow.ToSharedRef();
+							}
+						}
+
+						BaseActor = AttachParent;
+					}
+
+					// Now consider folder hierarchy for the base actor if desired
+					if (GetTreeViewMode() == EObjectMixerTreeViewMode::Folder ||
+						GetTreeViewMode() == EObjectMixerTreeViewMode::FolderObjectSubObject)
+					{
+						if (BaseActor)
+						{
+							FFolder BaseActorFolder = BaseActor->GetFolder();
+
+							while (!BaseActorFolder.IsNone())
+							{
+								TSharedPtr<FObjectMixerEditorListRow> FolderRow;
+							
+								if (const TSharedRef<FObjectMixerEditorListRow>* Match = FolderMap.Find(BaseActorFolder.GetPath()))
+								{
+									FolderRow = *Match;
+								}
+								else
+								{
+									FolderRow =
+										MakeShared<FObjectMixerEditorListRow>(
+											nullptr, FObjectMixerEditorListRow::Folder, SharedThis(this),
+											FText::FromName(BaseActorFolder.GetLeafName()));
+
+									FolderMap.Add(BaseActorFolder.GetPath(), FolderRow.ToSharedRef());
+								}
+
+								if (FolderRow)
+								{
+									FolderRow->AddToChildRows(TopLevelRow);
+													
+									TopLevelRow = FolderRow.ToSharedRef();
+								}
+
+								BaseActorFolder = BaseActorFolder.GetParent();
+							}
+						}
+					}
+				}
+			}
+
+			TreeViewRootObjects.AddUnique(TopLevelRow);
+		}
+	}
+
+	TreeViewRootObjects.StableSort(SortByTypeThenName);
+
+	if (TSharedPtr<FObjectMixerEditorMainPanel> MainPanel = ListModelPtr.Pin()->GetMainPanelModel().Pin())
+	{
+		MainPanel->RebuildCollectionSelector();
+	}
+
+	RefreshList();
+
+	RestoreTreeState();
+}
+
+void SObjectMixerEditorList::FindVisibleTreeViewObjects()
+{
+	VisibleTreeViewObjects.Empty();
+
+	for (const TSharedPtr<FObjectMixerEditorListRow>& Row : TreeViewRootObjects)
+	{
+		if (Row->ShouldRowWidgetBeVisible())
+		{
+			VisibleTreeViewObjects.Add(Row);
+		}
+	}
+}
+
+void SObjectMixerEditorList::FindVisibleObjectsAndRequestTreeRefresh()
+{
+	FindVisibleTreeViewObjects();
+	TreeViewPtr->RequestTreeRefresh();
+}
+
+void SObjectMixerEditorList::SelectedTreeItemsToSelectedInLevelEditor() const
+{
+	if (GEditor)
+	{
+		TArray<AActor*> ActorsToSelect;
+		for (const TSharedPtr<FObjectMixerEditorListRow>& SelectedRow : TreeViewPtr->GetSelectedItems())
+		{
+			if (SelectedRow->GetRowType() == FObjectMixerEditorListRow::MatchingObject ||
+				SelectedRow->GetRowType() == FObjectMixerEditorListRow::ContainerObject)
+			{
+				AActor* Actor = Cast<AActor>(SelectedRow->GetObject());
+
+				if (!Actor)
+				{
+					Actor = SelectedRow->GetObject()->GetTypedOuter<AActor>();
+				}
+
+				if (Actor)
+				{
+					ActorsToSelect.Add(Actor);
+				}
+			}
+		}
+
+		if (ActorsToSelect.Num())
+		{
+			GEditor->SelectNone(true, true, true);
+
+			for (AActor* Actor : ActorsToSelect)
+			{
+				GEditor->SelectActor( Actor, true, true, true );
+			}
+		}
+	}
+}
+
+void SObjectMixerEditorList::OnGetRowChildren(FObjectMixerEditorListRowPtr Row, TArray<FObjectMixerEditorListRowPtr>& OutChildren) const
+{
+	if (Row.IsValid())
+	{
+		OutChildren = Row->GetChildRows();
+
+		if (Row->GetShouldExpandAllChildren())
+		{
+			SetChildExpansionRecursively(Row, true);
+			Row->SetShouldExpandAllChildren(false);
+		}
+	}
+}
+
+void SObjectMixerEditorList::OnRowChildExpansionChange(FObjectMixerEditorListRowPtr Row, const bool bIsExpanded, const bool bIsRecursive) const
+{
+	if (Row.IsValid())
+	{
+		if (bIsRecursive)
+		{
+			if (bIsExpanded)
+			{
+				if (Row->GetRowType() == FObjectMixerEditorListRow::Folder)
+				{
+					Row->SetShouldExpandAllChildren(true);
+				}
+			}
+			else
+			{
+				SetChildExpansionRecursively(Row, bIsExpanded);
+			}
+		}
+		
+		TreeViewPtr->SetItemExpansion(Row, bIsExpanded);
+		Row->SetIsTreeViewItemExpanded(bIsExpanded);
+	}
+}
+
+void SObjectMixerEditorList::SetChildExpansionRecursively(const FObjectMixerEditorListRowPtr& InRow, const bool bNewIsExpanded) const
+{
+	if (InRow.IsValid())
+	{
+		for (const FObjectMixerEditorListRowPtr& Child : InRow->GetChildRows())
+		{
+			TreeViewPtr->SetItemExpansion(Child, bNewIsExpanded);
+			Child->SetIsTreeViewItemExpanded(bNewIsExpanded);
+
+			SetChildExpansionRecursively(Child, bNewIsExpanded);
+		}
+	}
+}
+
+#undef LOCTEXT_NAMESPACE
diff --git a/Engine/Plugins/Editor/ObjectMixer/ObjectMixer/Source/ObjectMixer/Private/Views/List/SObjectMixerEditorListRow.cpp b/Engine/Plugins/Editor/ObjectMixer/ObjectMixer/Source/ObjectMixer/Private/Views/List/SObjectMixerEditorListRow.cpp
new file mode 100644
--- /dev/null
+++ b/Engine/Plugins/Editor/ObjectMixer/ObjectMixer/Source/ObjectMixer/Private/Views/List/SObjectMixerEditorListRow.cpp
@@ -0,0 +1,536 @@
+// Copyright Epic Games, Inc. All Rights Reserved.
+
+#include "Views/List/SObjectMixerEditorListRow.h"
+
+#include "EditorClassUtils.h"
+#include "ObjectMixerEditorStyle.h"
+#include "Views/List/ObjectMixerEditorList.h"
+#include "Views/List/SObjectMixerEditorList.h"
+
+#include "Customizations/ColorStructCustomization.h"
+#include "Engine/Blueprint.h"
+#include "GameFramework/Actor.h"
+#include "Input/DragAndDrop.h"
+#include "ISinglePropertyView.h"
+#include "Modules/ModuleManager.h"
+#include "PropertyEditorModule.h"
+#include "PropertyHandle.h"
+#include "ScopedTransaction.h"
+#include "Styling/AppStyle.h"
+#include "Styling/StyleColors.h"
+#include "Widgets/Layout/SBox.h"
+#include "Widgets/SBoxPanel.h"
+#include "Widgets/Layout/SScaleBox.h"
+#include "Widgets/Text/STextBlock.h"
+
+#define LOCTEXT_NAMESPACE "ObjectMixerEditor"
+
+const FText InsertFormatText = LOCTEXT("InsertAboveFormatText", "Insert {0} {1} {2}");
+const FText AboveText = LOCTEXT("AboveListItem", "above");
+const FText BelowText = LOCTEXT("BelowListItem", "below");
+const FText MultiDragFormatText = LOCTEXT("MultiDragFormatText", "{0} Items");
+
+void SObjectMixerEditorListRow::Construct(
+	const FArguments& InArgs, const TSharedRef<STableViewBase>& InOwnerTable,
+	const TWeakPtr<FObjectMixerEditorListRow> InRow)
+{
+	check(InRow.IsValid());
+
+	Item = InRow;
+	const FObjectMixerEditorListRowPtr PinnedItem = Item.Pin();
+
+	SMultiColumnTableRow<FObjectMixerEditorListRowPtr>::Construct(
+		FSuperRowType::FArguments()
+		.Padding(1.0f)
+		.OnCanAcceptDrop(this, &SObjectMixerEditorListRow::HandleCanAcceptDrop)
+		.OnAcceptDrop(this, &SObjectMixerEditorListRow::HandleAcceptDrop)
+		.OnDragDetected(this, &SObjectMixerEditorListRow::HandleDragDetected)
+		.OnDragLeave(this, &SObjectMixerEditorListRow::HandleDragLeave)
+	 	, InOwnerTable
+	);
+
+	const FName VisibleHoveredBrushName = TEXT("Level.VisibleHighlightIcon16x");
+	const FName VisibleNotHoveredBrushName = TEXT("Level.VisibleIcon16x");
+	const FName NotVisibleHoveredBrushName = TEXT("Level.NotVisibleHighlightIcon16x");
+	const FName NotVisibleNotHoveredBrushName = TEXT("Level.NotVisibleIcon16x");
+
+	VisibleHoveredBrush = FAppStyle::Get().GetBrush(VisibleHoveredBrushName);
+	VisibleNotHoveredBrush = FAppStyle::Get().GetBrush(VisibleNotHoveredBrushName);
+	NotVisibleHoveredBrush = FAppStyle::Get().GetBrush(NotVisibleHoveredBrushName);
+	NotVisibleNotHoveredBrush = FAppStyle::Get().GetBrush(NotVisibleNotHoveredBrushName);
+}
+
+TSharedRef<SWidget> SObjectMixerEditorListRow::GenerateWidgetForColumn(const FName& InColumnName)
+{
+	const FObjectMixerEditorListRowPtr PinnedItem = Item.Pin();
+
+	if (const TSharedPtr<SWidget> CellWidget = GenerateCells(InColumnName, PinnedItem))
+	{
+		if (InColumnName == SObjectMixerEditorList::ItemNameColumnName)
+		{
+			// The first column gets the tree expansion arrow for this row
+			return SNew(SBox)
+				.MinDesiredHeight(20)
+				[
+					SNew( SHorizontalBox )
+
+					+SHorizontalBox::Slot()
+					.AutoWidth()
+					.Padding(6, 0, 0, 0)
+					[
+						SNew( SExpanderArrow, SharedThis(this) ).IndentAmount(12)
+					]
+
+					+SHorizontalBox::Slot()
+					.FillWidth(1.0f)
+					[
+						CellWidget.ToSharedRef()
+					]
+				];
+		}
+		
+		return SNew(SBorder)
+				   .HAlign(HAlign_Fill)
+				   .VAlign(VAlign_Center)
+				   .BorderImage(GetBorderImage(PinnedItem->GetRowType()))
+			   [
+				   CellWidget.ToSharedRef()
+			   ];
+	}
+
+	return SNullWidget::NullWidget;
+}
+
+void SObjectMixerEditorListRow::OnMouseEnter(const FGeometry& MyGeometry, const FPointerEvent& MouseEvent)
+{
+	bIsHovered = true;
+	
+	SMultiColumnTableRow<FObjectMixerEditorListRowPtr>::OnMouseEnter(MyGeometry, MouseEvent);
+}
+
+void SObjectMixerEditorListRow::OnMouseLeave(const FPointerEvent& MouseEvent)
+{
+	bIsHovered = false;
+
+	SMultiColumnTableRow<FObjectMixerEditorListRowPtr>::OnMouseLeave(MouseEvent);
+}
+
+SObjectMixerEditorListRow::~SObjectMixerEditorListRow()
+{
+	Item.Reset();
+}
+
+FReply SObjectMixerEditorListRow::HandleDragDetected(const FGeometry& MyGeometry, const FPointerEvent& MouseEvent)
+{
+	TArray<FObjectMixerEditorListRowPtr> DraggedItems = Item.Pin()->GetSelectedTreeViewItems();
+	TSharedRef<FObjectMixerListRowDragDropOp> Operation =
+		FObjectMixerListRowDragDropOp::New(DraggedItems);
+
+	return FReply::Handled().BeginDragDrop(Operation);
+}
+
+void SObjectMixerEditorListRow::HandleDragLeave(const FDragDropEvent& DragDropEvent)
+{
+	if (TSharedPtr<FObjectMixerListRowDragDropOp> Operation =
+		DragDropEvent.GetOperationAs<FObjectMixerListRowDragDropOp>())
+	{
+		Operation->ResetToDefaultToolTip();
+	}
+}
+
+TOptional<EItemDropZone> SObjectMixerEditorListRow::HandleCanAcceptDrop(const FDragDropEvent& DragDropEvent,
+                                                                             EItemDropZone DropZone,
+                                                                             FObjectMixerEditorListRowPtr
+                                                                             TargetItem)
+{
+	TSharedPtr<FObjectMixerListRowDragDropOp> Operation =
+		DragDropEvent.GetOperationAs<FObjectMixerListRowDragDropOp>();
+
+	if (!Operation.IsValid())
+	{
+		return TOptional<EItemDropZone>();
+	}
+
+	Operation->SetToolTip(
+		LOCTEXT("SortByCustomOrderDrgDropWarning", "Sort by custom order (\"#\") to drag & drop"),
+		FAppStyle::Get().GetBrush("Graph.ConnectorFeedback.Error")
+	);
+
+	const bool bIsDropDenied = false;
+
+	FString BoolAsString = bIsDropDenied ? "true" : "false";
+
+	if (bIsDropDenied)
+	{
+		Operation->ResetToDefaultToolTip();
+
+		return TOptional<EItemDropZone>();
+	}
+
+	FText ItemNameText = FText::FromString("Light Label");
+
+	if (Operation->DraggedItems.Num() > 1)
+	{
+		ItemNameText = FText::Format(MultiDragFormatText, FText::AsNumber(Operation->DraggedItems.Num()));
+	}
+
+	const FText DropPermittedText =
+		FText::Format(InsertFormatText,
+		              ItemNameText,
+		              DropZone == EItemDropZone::BelowItem ? BelowText : AboveText,
+		              ItemNameText
+		);
+
+	Operation->SetToolTip(
+		DropPermittedText,
+		FAppStyle::Get().GetBrush("Graph.ConnectorFeedback.OK")
+	);
+
+	// We have no behaviour yet for dropping one item onto another, so we'll treat it like we dropped it above
+	return DropZone == EItemDropZone::OntoItem ? EItemDropZone::AboveItem : DropZone;
+}
+
+FReply SObjectMixerEditorListRow::HandleAcceptDrop(const FDragDropEvent& DragDropEvent, EItemDropZone DropZone,
+                                                        FObjectMixerEditorListRowPtr TargetItem)
+{
+	TSharedPtr<FObjectMixerListRowDragDropOp> Operation =
+		DragDropEvent.GetOperationAs<FObjectMixerListRowDragDropOp>();
+
+	if (!Operation.IsValid())
+	{
+		return FReply::Unhandled();
+	}
+
+	const TSharedPtr<SObjectMixerEditorList> ListView = Item.Pin()->GetListViewPtr().Pin();
+
+	TArray<FObjectMixerEditorListRowPtr> DraggedItems = Operation->DraggedItems;
+
+	TArray<FObjectMixerEditorListRowPtr> AllTreeItemsCopy = ListView->GetTreeViewItems();
+
+	for (const FObjectMixerEditorListRowPtr& DraggedItem : DraggedItems)
+	{
+		if (!DraggedItem.IsValid() || !AllTreeItemsCopy.Contains(DraggedItem))
+		{
+			continue;
+		}
+
+		AllTreeItemsCopy.Remove(DraggedItem);
+	}
+
+	const int32 TargetIndex = AllTreeItemsCopy.IndexOfByKey(TargetItem);
+
+	if (TargetIndex > -1)
+	{
+		for (int32 ItemIndex = DraggedItems.Num() - 1; ItemIndex >= 0; ItemIndex--)
+		{
+			const FObjectMixerEditorListRowPtr& DraggedItem = DraggedItems[ItemIndex];
+
+			if (!DraggedItem.IsValid() || AllTreeItemsCopy.Contains(DraggedItem))
+			{
+				continue;
+			}
+
+			AllTreeItemsCopy.Insert(DraggedItem, DropZone == EItemDropZone::AboveItem ? TargetIndex : TargetIndex + 1);
+		}
+
+		ListView->SetTreeViewItems(AllTreeItemsCopy);
+	}
+
+	return FReply::Handled();
+}
+
+bool SObjectMixerEditorListRow::IsVisible() const
+{
+	if (const TSharedPtr<FObjectMixerEditorListRow> PinnedItem = Item.Pin())
+	{
+		return PinnedItem->GetObjectVisibility();
+	}
+
+	return false;
+}
+
+FSlateColor SObjectMixerEditorListRow::GetVisibilityIconForegroundColor() const
+{
+	check(Item.IsValid());
+	
+	const bool bIsSelected = Item.Pin()->GetIsSelected();
+
+	// make the foreground brush transparent if it is not selected and it is visible
+	if (IsVisible() && !bIsHovered && !bIsSelected)
+	{
+		return FLinearColor::Transparent;
+	}
+	else if (bIsHovered && !bIsSelected)
+	{
+		return FAppStyle::Get().GetSlateColor("Colors.ForegroundHover");
+	}
+
+	return FSlateColor::UseForeground();
+}
+
+FSlateColor SObjectMixerEditorListRow::GetSoloIconForegroundColor() const
+{
+	check(Item.IsValid());
+
+	if (Item.Pin()->IsThisRowSolo())
+	{
+		return bIsHovered ? FAppStyle::Get().GetSlateColor("Colors.ForegroundHover") : FSlateColor::UseForeground();
+	}
+
+	return FLinearColor::Transparent;
+}
+
+const FSlateBrush* SObjectMixerEditorListRow::GetVisibilityBrush() const
+{
+	if (IsVisible())
+	{
+		return bIsHovered ? VisibleHoveredBrush : VisibleNotHoveredBrush;
+	}
+	
+	return bIsHovered ? NotVisibleHoveredBrush : NotVisibleNotHoveredBrush;
+}
+
+const FSlateBrush* SObjectMixerEditorListRow::GetBorderImage(
+	const FObjectMixerEditorListRow::EObjectMixerEditorListRowType InRowType)
+{
+	return FObjectMixerEditorStyle::Get().GetBrush("ObjectMixerEditor.DefaultBorder");
+}
+
+TSharedPtr<SWidget> SObjectMixerEditorListRow::GenerateCells(
+	const FName& InColumnName, const TSharedPtr<FObjectMixerEditorListRow> PinnedItem)
+{
+	check(PinnedItem.IsValid());
+	
+	if (InColumnName.IsEqual(SObjectMixerEditorList::ItemNameColumnName))
+	{
+		TSharedRef<SHorizontalBox> HBox = SNew(SHorizontalBox);
+		
+		HBox->AddSlot()
+		.AutoWidth()
+		[
+			SNew(SImage)
+			.Image_Lambda([PinnedItem]()
+			{
+				return PinnedItem->GetObjectIconBrush();
+			})
+			.ColorAndOpacity(FSlateColor::UseForeground())
+		];
+
+		bool bNeedsStandardTextBlock = true;
+		const FText DisplayName = PinnedItem->GetDisplayName();
+		if (TObjectPtr<UObject> Object = PinnedItem->GetObject())
+		{
+			if (UClass* ActorClass = Object->GetClass())
+			{
+				if (UBlueprint* AsBlueprint = UBlueprint::GetBlueprintFromClass(ActorClass))
+				{
+					bNeedsStandardTextBlock = false;
+					
+					FEditorClassUtils::FSourceLinkParams SourceLinkParams;
+					SourceLinkParams.Object = Object;
+					SourceLinkParams.bUseDefaultFormat = false;
+					SourceLinkParams.bUseFormatIfNoLink = true;
+					SourceLinkParams.BlueprintFormat = &DisplayName;
+
+					HBox->AddSlot()
+					.Padding(FMargin(10.0, 0, 0, 0))
+					[
+						FEditorClassUtils::GetSourceLink(ActorClass, SourceLinkParams)
+					];
+				}
+			}
+		}
+
+		if (bNeedsStandardTextBlock)
+		{
+			HBox->AddSlot()
+			.Padding(FMargin(10.0, 0, 0, 0))
+			[
+				SNew(STextBlock)
+				.Visibility(EVisibility::Visible)
+				.Justification(ETextJustify::Left)
+				.Text(DisplayName)
+				.ToolTipText(DisplayName)
+			];
+		}
+		
+		return SNew(SBox)
+				.Visibility(EVisibility::SelfHitTestInvisible)
+				.HAlign(HAlign_Left)
+				.VAlign(VAlign_Center)
+				.Padding(FMargin(TextBlockLeftPadding, 0.f, 0.f, 0.f))
+				[
+					HBox
+				];
+	}
+
+	if (InColumnName.IsEqual(SObjectMixerEditorList::EditorVisibilityColumnName))
+	{
+		if (PinnedItem->GetRowType() == FObjectMixerEditorListRow::None)
+		{
+			return nullptr;
+		}
+
+		if (PinnedItem->GetObject() && !PinnedItem->GetObject()->IsA(AActor::StaticClass()))
+		{
+			return nullptr;
+		}
+		
+		return SNew(SBox)
+				.HAlign(HAlign_Left)
+				.VAlign(VAlign_Center)
+				.Visibility(EVisibility::SelfHitTestInvisible)
+				.Padding(0.f)
+				[
+					SNew(SImage)
+					.ColorAndOpacity(this, &SObjectMixerEditorListRow::GetVisibilityIconForegroundColor)
+					.Image_Raw(this, &SObjectMixerEditorListRow::GetVisibilityBrush)
+					.OnMouseButtonDown_Lambda(
+						[PinnedItem] (const FGeometry& MyGeometry, const FPointerEvent& Event)
+						{
+							check (PinnedItem);
+
+							FScopedTransaction Transaction( LOCTEXT("VisibilityChanged", "Object Mixer - Visibility Changed") );
+
+							const bool bIsVisible = PinnedItem->GetObjectVisibility();
+							
+							if (const TSharedPtr<SObjectMixerEditorList> PinnedListView = PinnedItem->GetListViewPtr().Pin();
+								PinnedItem->GetIsSelected() && PinnedListView->GetSelectedTreeViewItemCount() > 0)
+							{
+								PinnedListView->SetSelectedTreeViewItemActorsEditorVisible(!bIsVisible);
+
+								return FReply::Handled();
+							}
+
+							// Set Visibility Recursively
+							PinnedItem->SetObjectVisibility(!bIsVisible, true);
+
+							return FReply::Handled();
+						}
+					)
+				]
+			;
+	}
+
+	if (InColumnName.IsEqual(SObjectMixerEditorList::EditorVisibilitySoloColumnName))
+	{
+		if (PinnedItem->GetRowType() == FObjectMixerEditorListRow::None)
+		{
+			return nullptr;
+		}
+
+		if (PinnedItem->GetObject() && !PinnedItem->GetObject()->IsA(AActor::StaticClass()))
+		{
+			return nullptr;
+		}
+		
+		return SNew(SBox)
+				.HAlign(HAlign_Left)
+				.VAlign(VAlign_Center)
+				.Visibility(EVisibility::SelfHitTestInvisible)
+				.Padding(0.f)
+				[
+					SNew(SImage)
+					.ColorAndOpacity(this, &SObjectMixerEditorListRow::GetSoloIconForegroundColor)
+					.Image(FAppStyle::Get().GetBrush("MediaAsset.AssetActions.Solo.Small"))
+					.OnMouseButtonDown_Lambda(
+						[PinnedItem] (const FGeometry& MyGeometry, const FPointerEvent& Event)
+						{
+							check (PinnedItem);
+							
+							if (const TSharedPtr<SObjectMixerEditorList> PinnedListView = PinnedItem->GetListViewPtr().Pin();
+								PinnedListView->GetTreeViewItemCount() > 0)
+							{
+								const bool bIsRowSolo = PinnedItem->IsThisRowSolo();
+								for (const TSharedPtr<FObjectMixerEditorListRow>& TreeItem : PinnedListView->GetTreeViewItems())
+								{
+									TreeItem->SetObjectVisibility(bIsRowSolo, true);
+								}
+
+								if (bIsRowSolo)
+								{
+									PinnedItem->ClearSoloRow();
+								}
+								else
+								{
+									PinnedItem->SetObjectVisibility(true, true);
+									PinnedItem->SetThisAsSoloRow();
+								}
+
+								return FReply::Handled();
+							}
+
+							return FReply::Unhandled();
+						}
+					)
+				]
+			;
+	}
+	
+	if (UObject* ObjectRef = PinnedItem->GetObject())
+	{
+		FPropertyEditorModule& PropertyEditorModule = FModuleManager::LoadModuleChecked<FPropertyEditorModule>("PropertyEditor");
+		{
+			FSinglePropertyParams Params;
+			Params.NamePlacement = EPropertyNamePlacement::Hidden;
+			
+			const TSharedPtr<ISinglePropertyView> SinglePropertyView =
+				PropertyEditorModule.CreateSingleProperty(ObjectRef, InColumnName, Params
+			);
+
+			if (SinglePropertyView)
+			{
+				if (const TSharedPtr<IPropertyHandle> Handle = SinglePropertyView->GetPropertyHandle())
+				{
+					if (const FProperty* Property = Handle->GetProperty())
+					{
+						// Simultaneously edit all selected rows with a similar property
+						FSimpleDelegate OnPropertyValueChanged =
+							FSimpleDelegate::CreateRaw(
+								this,
+								&SObjectMixerEditorListRow::OnPropertyChanged, Property, (void*)ObjectRef);
+					
+						SinglePropertyView->SetOnPropertyValueChanged(OnPropertyValueChanged);
+
+						return SNew(SBox)
+								.Visibility(EVisibility::SelfHitTestInvisible)
+								.HAlign(HAlign_Fill)
+								.VAlign(VAlign_Center)
+								[
+									SinglePropertyView.ToSharedRef()
+								];
+					}
+				}
+			}
+		}
+	}
+
+	return nullptr;
+}
+
+void SObjectMixerEditorListRow::OnPropertyChanged(const FProperty* Property, void* ContainerWithChangedProperty)
+{
+	if (Property && ContainerWithChangedProperty)
+	{
+		if (const TSharedPtr<FObjectMixerEditorListRow> PinnedItem = Item.Pin())
+		{
+			if (PinnedItem->GetIsSelected())
+			{
+				if (const void* ValuePtr = Property->ContainerPtrToValuePtr<void>(ContainerWithChangedProperty))
+				{
+					FScopedTransaction Transaction( LOCTEXT("PropertyChanged", "Object Mixer - Property Changed") );
+					
+					for (const TSharedPtr<FObjectMixerEditorListRow>& SelectedRow : PinnedItem->GetSelectedTreeViewItems())
+					{
+						if (UObject* SelectedRowObject = SelectedRow->GetObject())
+						{
+							Property->SetValue_InContainer(SelectedRowObject, ValuePtr);
+						}
+					}
+				}
+			}
+		}
+	}
+}
+
+#undef LOCTEXT_NAMESPACE
diff --git a/Engine/Plugins/Editor/ObjectMixer/ObjectMixer/Source/ObjectMixer/Private/Views/MainPanel/ObjectMixerEditorMainPanel.cpp b/Engine/Plugins/Editor/ObjectMixer/ObjectMixer/Source/ObjectMixer/Private/Views/MainPanel/ObjectMixerEditorMainPanel.cpp
new file mode 100644
--- /dev/null
+++ b/Engine/Plugins/Editor/ObjectMixer/ObjectMixer/Source/ObjectMixer/Private/Views/MainPanel/ObjectMixerEditorMainPanel.cpp
@@ -0,0 +1,194 @@
+// Copyright Epic Games, Inc. All Rights Reserved.
+
+#include "Views/MainPanel/ObjectMixerEditorMainPanel.h"
+
+#include "ObjectMixerEditorSerializedData.h"
+#include "Views/List/ObjectMixerEditorList.h"
+#include "Views/MainPanel/SObjectMixerEditorMainPanel.h"
+
+void FObjectMixerEditorMainPanel::Init()
+{
+	RegenerateListModel();
+}
+
+TSharedRef<SWidget> FObjectMixerEditorMainPanel::GetOrCreateWidget()
+{
+	if (!MainPanelWidget.IsValid())
+	{
+		SAssignNew(MainPanelWidget, SObjectMixerEditorMainPanel, SharedThis(this));
+	}
+
+	return MainPanelWidget.ToSharedRef();
+}
+
+void FObjectMixerEditorMainPanel::RegenerateListModel()
+{
+	EditorListModel.Reset();
+	
+	EditorListModel = MakeShared<FObjectMixerEditorList>(SharedThis(this));
+}
+
+void FObjectMixerEditorMainPanel::RequestRebuildList() const
+{
+	if (EditorListModel.IsValid())
+	{
+		EditorListModel->RequestRebuildList();
+	}
+}
+
+void FObjectMixerEditorMainPanel::RefreshList() const
+{
+	if (EditorListModel.IsValid())
+	{
+		EditorListModel->RefreshList();
+	}
+}
+
+void FObjectMixerEditorMainPanel::RebuildCollectionSelector()
+{
+	MainPanelWidget->RebuildCollectionSelector();
+}
+
+FString FObjectMixerEditorMainPanel::GetSearchStringFromSearchInputField() const
+{
+	return MainPanelWidget->GetSearchStringFromSearchInputField();
+}
+
+void FObjectMixerEditorMainPanel::OnClassSelectionChanged(UClass* InNewClass)
+{
+	SetObjectFilterClass(InNewClass);
+}
+
+TObjectPtr<UClass> FObjectMixerEditorMainPanel::GetClassSelection() const
+{
+	return GetObjectFilterClass();
+}
+
+bool FObjectMixerEditorMainPanel::IsClassSelected(UClass* InNewClass) const
+{
+	return InNewClass == GetClassSelection();
+}
+
+UObjectMixerObjectFilter* FObjectMixerEditorMainPanel::GetObjectFilter()
+{
+	if (!ObjectFilterPtr.IsValid())
+	{
+		CacheObjectFilterObject();
+	}
+
+	return ObjectFilterPtr.Get();
+}
+
+void FObjectMixerEditorMainPanel::CacheObjectFilterObject()
+{
+	if (ObjectFilterPtr.IsValid())
+	{
+		ObjectFilterPtr.Reset();
+	}
+	
+	if (const UClass* Class = GetObjectFilterClass())
+	{
+		ObjectFilterPtr = TStrongObjectPtr(NewObject<UObjectMixerObjectFilter>(GetTransientPackage(), Class));
+	}
+}
+
+const TArray<TSharedRef<IObjectMixerEditorListFilter>>& FObjectMixerEditorMainPanel::GetShowFilters() const
+{
+	return MainPanelWidget->GetShowFilters();	
+}
+
+void FObjectMixerEditorMainPanel::AddObjectsToCollection(const FName& CollectionName, const TSet<FSoftObjectPath>& ObjectsToAdd) const
+{
+	if (const TSubclassOf<UObjectMixerObjectFilter> Filter = GetObjectFilterClass())
+	{
+		if (UObjectMixerEditorSerializedData* SerializedData = GetMutableDefault<UObjectMixerEditorSerializedData>())
+		{
+			SerializedData->AddObjectsToCollection(Filter->GetFName(), CollectionName, ObjectsToAdd);
+
+			OnObjectMixerCollectionMapChanged.Broadcast();
+		}
+	}
+}
+
+void FObjectMixerEditorMainPanel::RemoveObjectsFromCollection(const FName& CollectionName, const TSet<FSoftObjectPath>& ObjectsToRemove) const
+{
+	if (const TSubclassOf<UObjectMixerObjectFilter> Filter = GetObjectFilterClass())
+	{
+		if (UObjectMixerEditorSerializedData* SerializedData = GetMutableDefault<UObjectMixerEditorSerializedData>())
+		{
+			SerializedData->RemoveObjectsFromCollection(Filter->GetFName(), CollectionName, ObjectsToRemove);
+
+			OnObjectMixerCollectionMapChanged.Broadcast();
+		}
+	}
+}
+
+void FObjectMixerEditorMainPanel::RemoveCollection(const FName& CollectionName) const
+{
+	if (const TSubclassOf<UObjectMixerObjectFilter> Filter = GetObjectFilterClass())
+	{
+		if (UObjectMixerEditorSerializedData* SerializedData = GetMutableDefault<UObjectMixerEditorSerializedData>())
+		{
+			SerializedData->RemoveCollection(Filter->GetFName(), CollectionName);
+
+			OnObjectMixerCollectionMapChanged.Broadcast();
+		}
+	}
+}
+
+void FObjectMixerEditorMainPanel::ReorderCollection(const FName& CollectionToMoveName, const FName& CollectionInsertBeforeName) const
+{
+	if (const TSubclassOf<UObjectMixerObjectFilter> Filter = GetObjectFilterClass())
+	{
+		if (UObjectMixerEditorSerializedData* SerializedData = GetMutableDefault<UObjectMixerEditorSerializedData>())
+		{
+			SerializedData->ReorderCollection(Filter->GetFName(), CollectionToMoveName, CollectionInsertBeforeName);
+
+			OnObjectMixerCollectionMapChanged.Broadcast();
+		}
+	}
+}
+
+bool FObjectMixerEditorMainPanel::IsObjectInCollection(const FName& CollectionName, const FSoftObjectPath& InObject) const
+{
+	if (const TSubclassOf<UObjectMixerObjectFilter> Filter = GetObjectFilterClass())
+	{
+		if (UObjectMixerEditorSerializedData* SerializedData = GetMutableDefault<UObjectMixerEditorSerializedData>())
+		{
+			return SerializedData->IsObjectInCollection(Filter->GetFName(), CollectionName, InObject);
+		}
+	}
+
+	return false;
+}
+
+TSet<FName> FObjectMixerEditorMainPanel::GetCollectionsForObject(const FSoftObjectPath& InObject) const
+{
+	if (const TSubclassOf<UObjectMixerObjectFilter> Filter = GetObjectFilterClass())
+	{
+		if (UObjectMixerEditorSerializedData* SerializedData = GetMutableDefault<UObjectMixerEditorSerializedData>())
+		{
+			return SerializedData->GetCollectionsForObject(Filter->GetFName(), InObject);
+		}
+	}
+
+	return {};
+}
+
+TArray<FName> FObjectMixerEditorMainPanel::GetAllCollectionNames() const
+{
+	if (const TSubclassOf<UObjectMixerObjectFilter> Filter = GetObjectFilterClass())
+	{
+		if (UObjectMixerEditorSerializedData* SerializedData = GetMutableDefault<UObjectMixerEditorSerializedData>())
+		{
+			return SerializedData->GetAllCollectionNames(Filter->GetFName());
+		}
+	}
+
+	return {};
+}
+
+const TSet<FName>& FObjectMixerEditorMainPanel::GetCurrentCollectionSelection() const
+{
+	return MainPanelWidget->GetCurrentCollectionSelection();
+}
diff --git a/Engine/Plugins/Editor/ObjectMixer/ObjectMixer/Source/ObjectMixer/Private/Views/MainPanel/SObjectMixerEditorMainPanel.cpp b/Engine/Plugins/Editor/ObjectMixer/ObjectMixer/Source/ObjectMixer/Private/Views/MainPanel/SObjectMixerEditorMainPanel.cpp
new file mode 100644
--- /dev/null
+++ b/Engine/Plugins/Editor/ObjectMixer/ObjectMixer/Source/ObjectMixer/Private/Views/MainPanel/SObjectMixerEditorMainPanel.cpp
@@ -0,0 +1,585 @@
+// Copyright Epic Games, Inc. All Rights Reserved.
+
+#include "Views/MainPanel/SObjectMixerEditorMainPanel.h"
+
+#include "ObjectMixerEditorLog.h"
+#include "ObjectMixerEditorModule.h"
+#include "ObjectMixerEditorProjectSettings.h"
+#include "Views/List/ObjectMixerEditorList.h"
+#include "Views/MainPanel/ObjectMixerEditorMainPanel.h"
+#include "Views/Widgets/SCollectionSelectionButton.h"
+#include "Views/Widgets/SObjectMixerPlacementAssetMenuEntry.h"
+
+#include "Framework/MultiBox/MultiBoxBuilder.h"
+#include "Kismet2/SClassPickerDialog.h"
+#include "PlacementMode/Public/IPlacementModeModule.h"
+#include "SPositiveActionButton.h"
+#include "Brushes/SlateRoundedBoxBrush.h"
+#include "Styling/StyleColors.h"
+#include "Views/List/ObjectMixerEditorListFilters/IObjectMixerEditorListFilter.h"
+#include "Views/List/ObjectMixerEditorListFilters/ObjectMixerEditorListFilter_Collection.h"
+#include "Widgets/Input/SCheckBox.h"
+#include "Widgets/Input/SComboButton.h"
+#include "Widgets/Input/SMultiLineEditableTextBox.h"
+#include "Widgets/Input/SSearchBox.h"
+#include "Widgets/Layout/SWidgetSwitcher.h"
+#include "Widgets/Layout/SWrapBox.h"
+#include "Widgets/Text/STextBlock.h"
+
+#define LOCTEXT_NAMESPACE "ObjectMixerEditor"
+
+void SObjectMixerEditorMainPanel::Construct(
+	const FArguments& InArgs, const TSharedRef<FObjectMixerEditorMainPanel>& InMainPanel)
+{
+	check(InMainPanel->GetEditorListModel().IsValid());
+
+	MainPanelModel = InMainPanel;
+	
+	ChildSlot
+	[
+		SNew(SVerticalBox)
+
+		+SVerticalBox::Slot()
+		.HAlign(HAlign_Fill)
+		.AutoHeight()
+		[
+			GenerateToolbar()
+		]
+
+		+SVerticalBox::Slot()
+		.Padding(8, 2, 8, 7)
+		.AutoHeight()
+		[
+			SAssignNew(CollectionSelectorBox, SWrapBox)
+			.UseAllottedSize(true)
+			.InnerSlotPadding(FVector2D(4,4))
+		]
+
+		+SVerticalBox::Slot()
+		.HAlign(HAlign_Fill)
+		.VAlign(VAlign_Fill)
+		[
+			MainPanelModel.Pin()->GetEditorListModel().Pin()->GetOrCreateWidget()
+		]
+	];
+
+	InMainPanel->GetOnObjectMixerCollectionMapChanged().AddRaw(this, &SObjectMixerEditorMainPanel::RebuildCollectionSelector);
+
+	ShowFilters.Add(MakeShared<FObjectMixerEditorListFilter_Collection>());
+}
+
+TSharedRef<SWidget> SObjectMixerEditorMainPanel::GenerateToolbar()
+{
+	TSharedRef<SHorizontalBox> ToolbarBox = SNew(SHorizontalBox);
+
+	// Add object button
+	ToolbarBox->AddSlot()
+   .HAlign(HAlign_Left)
+   .VAlign(VAlign_Center)
+   .AutoWidth()
+   .Padding(FMargin(0, 4))
+   [
+	   SNew(SPositiveActionButton)
+	   .Text(LOCTEXT("AddObject", "Add"))
+	   .OnGetMenuContent(FOnGetContent::CreateRaw(this, &SObjectMixerEditorMainPanel::OnGenerateAddObjectButtonMenu))
+   ];
+
+	ToolbarBox->AddSlot()
+	.HAlign(HAlign_Fill)
+	.VAlign(VAlign_Center)
+	.Padding(0.f, 1.f, 0.f, 1.f)
+	[
+		SAssignNew(SearchBoxPtr, SSearchBox)
+		.HintText(LOCTEXT("SearchHintText", "Search Scene Objects"))
+		.ToolTipText(LOCTEXT("ObjectMixerEditorList_TooltipText", "Search Scene Objects"))
+		.OnTextChanged_Raw(this, &SObjectMixerEditorMainPanel::OnSearchTextChanged)
+	];
+
+	// Selection Sync Toggle
+	ToolbarBox->AddSlot()
+	.HAlign(HAlign_Right)
+	.VAlign(VAlign_Center)
+	.AutoWidth()
+	.Padding(8.f, 1.f, 0.f, 1.f)
+	[
+		SNew(SCheckBox )
+		.Padding(4.f)
+		.ToolTipText(LOCTEXT("SyncSelectionButton_Tooltip", "Sync Selection\nIf enabled, clicking an item in the mixer list will also select the item in the Scene Outliner.\nAlt + Click to select items in mixer without selecting the item in the Scene outliner.\nIf disabled, selections will not sync unless Alt is held. Effectively, this is the opposite behavior."))
+		.Style(&FAppStyle::Get().GetWidgetStyle<FCheckBoxStyle>("ToggleButtonCheckbox"))
+		.ForegroundColor(FSlateColor::UseForeground())
+		.IsChecked_Lambda([]()
+		{
+			if (const UObjectMixerEditorProjectSettings* Settings = GetDefault<UObjectMixerEditorProjectSettings>())
+			{
+				return Settings->bSyncSelection ? ECheckBoxState::Checked : ECheckBoxState::Unchecked;
+			}
+			
+			return ECheckBoxState::Undetermined; 
+		})
+		.OnCheckStateChanged_Lambda([](ECheckBoxState InNewState)
+		{
+			if (UObjectMixerEditorProjectSettings* Settings = GetMutableDefault<UObjectMixerEditorProjectSettings>())
+			{
+				Settings->bSyncSelection = InNewState == ECheckBoxState::Checked ? true : false;
+			}
+		})
+		[
+			SNew(SImage)
+			.ColorAndOpacity(FSlateColor::UseForeground())
+			.Image( FAppStyle::Get().GetBrush("FoliageEditMode.SelectAll") )
+		]
+	];
+
+	// Show Options
+	ToolbarBox->AddSlot()
+	.HAlign(HAlign_Right)
+	.VAlign(VAlign_Center)
+	.AutoWidth()
+	.Padding(8.f, 1.f, 0.f, 1.f)
+	[
+		SAssignNew( ViewOptionsComboButton, SComboButton )
+		.ContentPadding(4.f)
+		.ToolTipText(LOCTEXT("ShowOptions_Tooltip", "Show options to affect the visibility of items in the Object Mixer list"))
+		.ComboButtonStyle( FAppStyle::Get(), "SimpleComboButtonWithIcon" ) // Use the tool bar item style for this button
+		.OnGetMenuContent( this, &SObjectMixerEditorMainPanel::BuildShowOptionsMenu)
+		.HasDownArrow(false)
+		.ButtonContent()
+		[
+			SNew(SImage)
+			.ColorAndOpacity(FSlateColor::UseForeground())
+			.Image( FAppStyle::Get().GetBrush("Icons.Settings") )
+		]
+	];
+
+	return ToolbarBox;
+}
+
+TSharedRef<SWidget> SObjectMixerEditorMainPanel::OnGenerateAddObjectButtonMenu() const
+{
+	const TSet<TSubclassOf<AActor>> SubclassesOfActor = MainPanelModel.Pin()->GetObjectClassesToPlace();
+
+	if (SubclassesOfActor.Num() > 0)
+	{
+		TSet<UClass*> ClassesToPlace =
+		   MainPanelModel.Pin()->GetObjectFilter()->GetParentAndChildClassesFromSpecifiedClasses(
+			   SubclassesOfActor,
+			   MainPanelModel.Pin()->GetObjectFilter()->GetObjectMixerPlacementClassInclusionOptions());
+	
+		FMenuBuilder AddObjectButtonMenuBuilder = FMenuBuilder(true, nullptr);
+
+		for (const UClass* Class : ClassesToPlace)
+		{
+			if (const UActorFactory* Factory = GEditor->FindActorFactoryForActorClass(Class))
+			{
+				AddObjectButtonMenuBuilder.AddWidget(
+					SNew(SObjectMixerPlacementAssetMenuEntry, MakeShareable(new FPlaceableItem(*Factory->GetClass()))), FText::GetEmpty());
+			}
+		}
+
+		return AddObjectButtonMenuBuilder.MakeWidget();
+	}
+
+	return
+		SNew(SBox)
+		.Padding(5)
+		[
+			SNew(STextBlock)
+				.Text(LOCTEXT("NoPlaceableActorsDefinedWarning", "Please define some placeable actors in the\nfilter class by overriding GetObjectClassesToPlace."))
+				.Font(FAppStyle::Get().GetFontStyle("NormalFontItalic"))
+		]
+	;
+}
+
+TSharedRef<SWidget> SObjectMixerEditorMainPanel::OnGenerateFilterClassMenu()
+{
+	FMenuBuilder MenuBuilder(true, nullptr);
+
+	TArray<UClass*> DerivedClasses;
+	GetDerivedClasses(UObjectMixerObjectFilter::StaticClass(), DerivedClasses, true);
+
+	DerivedClasses.Remove(UObjectMixerObjectFilter::StaticClass());
+	DerivedClasses.Remove(UObjectMixerBlueprintObjectFilter::StaticClass());
+
+	DerivedClasses.Sort([](UClass& A, UClass& B)
+	{
+		return A.GetFName().LexicalLess(B.GetFName());
+	});
+
+	if (DerivedClasses.Num())
+	{
+		check(MainPanelModel.Pin());
+		
+		MenuBuilder.BeginSection(NAME_None, LOCTEXT("SelectClassMenuSection", "Select Class"));
+		{
+			for (UClass* DerivedClass : DerivedClasses)
+			{
+				if (IsValid(DerivedClass))
+				{
+					if (DerivedClass->GetName().StartsWith(TEXT("SKEL_")) || DerivedClass->GetName().StartsWith(TEXT("REINST_")))
+					{
+						continue;
+					}
+
+					if (DerivedClass->HasAnyClassFlags(CLASS_Abstract | CLASS_HideDropDown | CLASS_Deprecated))
+					{
+						continue;
+					}
+					
+					MenuBuilder.AddMenuEntry(
+					   FText::FromName(DerivedClass->GetFName()),
+					   FText::GetEmpty(),
+					   FSlateIcon(),
+					   FUIAction(
+						   FExecuteAction::CreateSP(MainPanelModel.Pin().ToSharedRef(), &FObjectMixerEditorMainPanel::SetObjectFilterClass, DerivedClass),
+						   FCanExecuteAction::CreateLambda([](){ return true; }),
+						   FIsActionChecked::CreateSP(MainPanelModel.Pin().ToSharedRef(), &FObjectMixerEditorMainPanel::IsClassSelected, DerivedClass)
+					   ),
+					   NAME_None,
+					   EUserInterfaceActionType::RadioButton
+				   );
+				}
+			}
+		}
+		MenuBuilder.EndSection();
+	}
+	else
+	{
+		MenuBuilder.AddMenuEntry(LOCTEXT("NoFilterClassesAvailable", "No filter classes available."), FText::GetEmpty(), FSlateIcon(), FUIAction());
+	}
+
+	TSharedRef<SWidget> Widget = MenuBuilder.MakeWidget();
+	FChildren* ChildWidgets = Widget->GetChildren();
+	for (int32 ChildItr = 0; ChildItr < ChildWidgets->Num(); ChildItr++)
+	{
+		const TSharedRef<SWidget>& Child = ChildWidgets->GetChildAt(ChildItr);
+
+		Child->EnableToolTipForceField(false);
+	}
+	Widget->EnableToolTipForceField(false);
+	
+	return Widget;
+}
+
+TSharedRef<SWidget> SObjectMixerEditorMainPanel::BuildShowOptionsMenu()
+{
+	FMenuBuilder ShowOptionsMenuBuilder = FMenuBuilder(true, nullptr);
+
+	// No need to select filter class from outside generic instance
+	if (MainPanelModel.Pin()->GetModuleName() == FObjectMixerEditorModule::BaseObjectMixerModuleName)
+	{
+		ShowOptionsMenuBuilder.BeginSection("ListViewOptions", LOCTEXT("FilterClassManagementSection", "Filter Class Management"));
+		{
+			// Filter Class Management Button
+			const TSharedRef<SWidget> FilterClassManagementButton =
+				SNew(SBox)
+				.Padding(8, 0)
+				[
+					SNew(SComboButton)
+					.ToolTipText(LOCTEXT("FilterClassManagementButton_Tooltip", "Select a filter class"))
+					.ContentPadding(FMargin(4, 0.5f))
+					.ComboButtonStyle(&FAppStyle::Get().GetWidgetStyle<FComboButtonStyle>("ComboButton"))
+					.OnGetMenuContent(this, &SObjectMixerEditorMainPanel::OnGenerateFilterClassMenu)
+					.ForegroundColor(FStyleColors::Foreground)
+					.MenuPlacement(EMenuPlacement::MenuPlacement_MenuRight)
+					.ButtonContent()
+					[
+						SNew(SHorizontalBox)
+
+						+ SHorizontalBox::Slot()
+						.Padding(0, 1, 4, 0)
+						.AutoWidth()
+						[
+							SNew(SImage)
+							.Image(FAppStyle::Get().GetBrush("Icons.Filter"))
+							.ColorAndOpacity(FSlateColor::UseForeground())
+						]
+
+						+ SHorizontalBox::Slot()
+						.Padding(0, 1, 0, 0)
+						.AutoWidth()
+						[
+							SNew(STextBlock)
+							.Text(LOCTEXT("FilterClassToolbarButton", "Object Filter Class"))
+						]
+					]
+				];
+
+			ShowOptionsMenuBuilder.AddWidget(FilterClassManagementButton, FText::GetEmpty());
+		}
+		ShowOptionsMenuBuilder.EndSection();
+	}
+
+	// Add List View Mode Options
+	ShowOptionsMenuBuilder.BeginSection("ListViewOptions", LOCTEXT("ListViewOptionsSection", "List View Options"));
+	{
+		// Foreach on uenum
+		const FString EnumPath = "/Script/ObjectMixerEditor.EObjectMixerTreeViewMode";
+		if (const UEnum* EnumPtr = FindObject<UEnum>(nullptr, *EnumPath, true))
+		{
+			for (int32 EnumItr = 0; EnumItr < EnumPtr->GetMaxEnumValue(); EnumItr++)
+			{
+				EObjectMixerTreeViewMode EnumValue = (EObjectMixerTreeViewMode)EnumItr;
+				
+				ShowOptionsMenuBuilder.AddMenuEntry(
+					EnumPtr->GetDisplayNameTextByIndex(EnumItr),
+					EnumPtr->GetToolTipTextByIndex(EnumItr),
+					FSlateIcon(),
+					FUIAction(
+						FExecuteAction::CreateLambda([this, EnumValue]()
+						{
+							SetTreeViewMode(EnumValue);
+						}),
+						FCanExecuteAction::CreateLambda([](){ return true; }),
+						FIsActionChecked::CreateLambda([this, EnumValue]()
+						{
+							return GetTreeViewMode() == EnumValue;
+						})
+					),
+					NAME_None,
+					EUserInterfaceActionType::RadioButton
+				);
+			}
+		}
+	}
+	ShowOptionsMenuBuilder.EndSection();
+
+	ShowOptionsMenuBuilder.BeginSection("MiscOptionsSection", LOCTEXT("MiscOptionsSection","Misc"));
+	{
+		// No need to open generic instance from itself
+		if (MainPanelModel.Pin()->GetModuleName() != FObjectMixerEditorModule::BaseObjectMixerModuleName)
+		{
+			ShowOptionsMenuBuilder.AddMenuEntry(
+			  LOCTEXT("OpenGenericInstanceMenuOption", "Open Generic Object Mixer Instance"),
+			  LOCTEXT("OpenGenericInstanceMenuOptionTooltip", "Open a generic object mixer instance that can take in a user-specified filter class."),
+			  FSlateIcon(),
+			  FUIAction(FExecuteAction::CreateLambda([]()
+			  {
+				  FGlobalTabmanager::Get()->TryInvokeTab(FObjectMixerEditorModule::Get().GetTabSpawnerId());
+			  })));
+		}
+
+		ShowOptionsMenuBuilder.AddMenuEntry(
+			LOCTEXT("RebuildListMenuOption","Rebuild List"), 
+			LOCTEXT("RebuildListMenuOptionTooltip","Force the active list to be rebuilt. Useful if the list doesn't automatically refresh."),
+			FSlateIcon(),
+			FUIAction(FExecuteAction::CreateSP(MainPanelModel.Pin().ToSharedRef(), &FObjectMixerEditorMainPanel::RequestRebuildList)));
+	}
+	ShowOptionsMenuBuilder.EndSection();
+
+	if (ShowFilters.Num())
+	{
+		ShowOptionsMenuBuilder.BeginSection("", LOCTEXT("ShowOptions_ShowSectionHeading", "Show"));
+		{
+			// Add show filters
+			auto AddFiltersLambda = [this, &ShowOptionsMenuBuilder](const TSharedRef<IObjectMixerEditorListFilter>& InFilter)
+			{
+				const FString& FilterName = InFilter->GetFilterName();
+			
+				ShowOptionsMenuBuilder.AddMenuEntry(
+				   InFilter->GetFilterButtonLabel(),
+				   InFilter->GetFilterButtonToolTip(),
+				   FSlateIcon(),
+				   FUIAction(
+					   FExecuteAction::CreateLambda(
+						   [this, FilterName]()
+						   {
+							   ToggleFilterActive(FilterName);
+						   }
+						),
+					   FCanExecuteAction(),
+					   FIsActionChecked::CreateSP( InFilter, &IObjectMixerEditorListFilter::GetIsFilterActive )
+				   ),
+				   NAME_None,
+				   EUserInterfaceActionType::ToggleButton
+			   );
+			};
+
+			for (const TSharedRef<IObjectMixerEditorListFilter>& Filter : ShowFilters)
+			{
+				// Don't add non-toggle filters to Show Options, they're always on
+				if (Filter->IsToggleable())
+				{
+					AddFiltersLambda(Filter);
+				}
+			}
+		}
+		ShowOptionsMenuBuilder.EndSection();
+	}
+
+	return ShowOptionsMenuBuilder.MakeWidget();
+}
+
+void SObjectMixerEditorMainPanel::OnSearchTextChanged(const FText& Text)
+{
+	ExecuteListViewSearchOnAllRows(Text.ToString(), true);
+}
+
+FString SObjectMixerEditorMainPanel::GetSearchStringFromSearchInputField() const
+{
+	return ensureAlwaysMsgf(SearchBoxPtr.IsValid(),
+		TEXT("%hs: SearchBoxPtr is not valid. Check to make sure it was created."), __FUNCTION__)
+	? SearchBoxPtr->GetText().ToString() : "";
+}
+
+void SObjectMixerEditorMainPanel::SetSearchStringInSearchInputField(const FString InSearchString) const
+{
+	if (ensureAlwaysMsgf(SearchBoxPtr.IsValid(),
+		TEXT("%hs: SearchBoxPtr is not valid. Check to make sure it was created."), __FUNCTION__))
+	{
+		SearchBoxPtr->SetText(FText::FromString(InSearchString));
+	}
+}
+
+void SObjectMixerEditorMainPanel::ExecuteListViewSearchOnAllRows(const FString& SearchString,
+	const bool bShouldRefreshAfterward)
+{
+	MainPanelModel.Pin()->GetEditorListModel().Pin()->ExecuteListViewSearchOnAllRows(GetSearchStringFromSearchInputField());
+}
+
+EObjectMixerTreeViewMode SObjectMixerEditorMainPanel::GetTreeViewMode()
+{
+	const TSharedPtr<FObjectMixerEditorMainPanel> PinnedPanelModel = MainPanelModel.Pin();
+	check(PinnedPanelModel);
+	
+	return PinnedPanelModel->GetTreeViewMode();
+}
+
+void SObjectMixerEditorMainPanel::SetTreeViewMode(EObjectMixerTreeViewMode InViewMode)
+{
+	if (TSharedPtr<FObjectMixerEditorMainPanel> PinnedPanelModel = MainPanelModel.Pin())
+	{
+		PinnedPanelModel->SetTreeViewMode(InViewMode);
+	}
+}
+
+void SObjectMixerEditorMainPanel::ToggleFilterActive(const FString& FilterName)
+{
+	if (const TSharedRef<IObjectMixerEditorListFilter>* Match =
+		Algo::FindByPredicate(GetShowFilters(),
+		[&FilterName](TSharedRef<IObjectMixerEditorListFilter> Comparator)
+		{
+			return Comparator->GetFilterName().Equals(FilterName);
+		}))
+	{
+		const TSharedRef<IObjectMixerEditorListFilter> Filter = *Match;
+		Filter->ToggleFilterActive();
+
+		MainPanelModel.Pin()->GetEditorListModel().Pin()->EvaluateIfRowsPassFilters();
+	}
+}
+
+const TSet<FName>& SObjectMixerEditorMainPanel::GetCurrentCollectionSelection()
+{
+	return CurrentCollectionSelection;
+}
+
+void SObjectMixerEditorMainPanel::RebuildCollectionSelector()
+{
+	check(MainPanelModel.IsValid());
+	
+	CollectionSelectorBox->ClearChildren();
+	CollectionSelectorBox->SetVisibility(EVisibility::Collapsed);
+
+	TArray<FName> AllCollections = MainPanelModel.Pin()->GetAllCollectionNames();
+	
+	if (AllCollections.IsEmpty())
+	{
+		// no collections - rather than show just "All", hide the box
+		ResetCurrentCollectionSelection();
+		return;
+	}
+
+	auto CreateSection = [this](FName CollectionName) 
+	{
+		return SNew(SCollectionSelectionButton, SharedThis(this), CollectionName);
+	};
+
+	for (const FName& Key : AllCollections)
+	{
+		CollectionSelectorBox->AddSlot()
+		[
+			CreateSection(Key)
+		];
+	}
+
+	CollectionSelectorBox->AddSlot()
+	[
+		CreateSection("All")
+	];
+
+	CollectionSelectorBox->SetVisibility(EVisibility::Visible);
+}
+
+void SObjectMixerEditorMainPanel::RequestRemoveCollection(const FName& CollectionName)
+{
+	MainPanelModel.Pin()->RemoveCollection(CollectionName);
+
+	CurrentCollectionSelection.Remove(CollectionName);
+}
+
+void SObjectMixerEditorMainPanel::OnCollectionCheckedStateChanged(ECheckBoxState State, FName CollectionName)
+{
+	check(MainPanelModel.IsValid());
+
+	// Remove Collection
+	if (FSlateApplication::Get().GetModifierKeys().IsAltDown())
+	{
+		RequestRemoveCollection(CollectionName);
+	}
+	else
+	{
+		const bool bIsControlDown = FSlateApplication::Get().GetModifierKeys().IsControlDown();
+
+		if (State == ECheckBoxState::Unchecked)
+		{
+			if (bIsControlDown)
+			{
+				CurrentCollectionSelection.Remove(CollectionName);
+			}
+			else
+			{
+				CurrentCollectionSelection.Reset();
+
+				if (CollectionName != "All")
+				{
+					CurrentCollectionSelection.Add(CollectionName);
+				}
+			}
+		}
+		else if (State == ECheckBoxState::Checked)
+		{
+			if (!bIsControlDown)
+			{
+				CurrentCollectionSelection.Reset();
+			}
+
+			if (CollectionName != "All")
+			{
+				CurrentCollectionSelection.Add(CollectionName);
+			}
+		}
+	}
+
+	MainPanelModel.Pin()->GetEditorListModel().Pin()->EvaluateIfRowsPassFilters();
+}
+
+ECheckBoxState SObjectMixerEditorMainPanel::IsCollectionChecked(FName Section) const
+{
+	if (CurrentCollectionSelection.IsEmpty() && Section == "All")
+	{
+		return ECheckBoxState::Checked;
+	}
+
+	return CurrentCollectionSelection.Contains(Section) ? ECheckBoxState::Checked : ECheckBoxState::Unchecked;
+}
+
+SObjectMixerEditorMainPanel::~SObjectMixerEditorMainPanel()
+{
+	if (MainPanelModel.IsValid())
+	{
+		MainPanelModel.Pin()->GetOnObjectMixerCollectionMapChanged().RemoveAll(this);
+		MainPanelModel.Reset();
+	}
+}
+
+#undef LOCTEXT_NAMESPACE
diff --git a/Engine/Plugins/Editor/ObjectMixer/ObjectMixer/Source/ObjectMixer/Private/Views/Widgets/ObjectMixerEditorListMenuContext.cpp b/Engine/Plugins/Editor/ObjectMixer/ObjectMixer/Source/ObjectMixer/Private/Views/Widgets/ObjectMixerEditorListMenuContext.cpp
new file mode 100644
--- /dev/null
+++ b/Engine/Plugins/Editor/ObjectMixer/ObjectMixer/Source/ObjectMixer/Private/Views/Widgets/ObjectMixerEditorListMenuContext.cpp
@@ -0,0 +1,247 @@
+// Copyright Epic Games, Inc. All Rights Reserved.
+
+#include "Views/Widgets/ObjectMixerEditorListMenuContext.h"
+
+#include "LevelEditor.h"
+#include "LevelEditorMenuContext.h"
+#include "Framework/MultiBox/MultiBoxDefs.h"
+#include "ToolMenu.h"
+#include "ToolMenus.h"
+#include "ToolMenuContext.h"
+#include "ToolMenuDelegates.h"
+#include "ToolMenuEntry.h"
+#include "ToolMenuSection.h"
+#include "UnrealEdGlobals.h"
+#include "Views/MainPanel/ObjectMixerEditorMainPanel.h"
+#include "Widgets/Layout/SBox.h"
+#include "Widgets/Input/SEditableTextBox.h"
+
+#define LOCTEXT_NAMESPACE "ObjectMixerEditorListMenuContext"
+
+FName UObjectMixerEditorListMenuContext::DefaultContextBaseMenuName("ObjectMixer.ContextMenuBase");
+
+TSharedPtr<SWidget> UObjectMixerEditorListMenuContext::CreateContextMenu(const FObjectMixerEditorListMenuContextData InData)
+{
+	if (InData.SelectedItems.Num() == 0)
+	{
+		return nullptr;
+	}
+
+	return BuildContextMenu(InData);
+}
+
+TSharedPtr<SWidget> UObjectMixerEditorListMenuContext::BuildContextMenu(const FObjectMixerEditorListMenuContextData& InData)
+{
+	RegisterContextMenu();
+
+	UObjectMixerEditorListMenuContext* ContextObject = NewObject<UObjectMixerEditorListMenuContext>();
+	ContextObject->Data = InData;
+
+	const FToolMenuContext Context(ContextObject);
+
+	// todo: Uncomment all code when Level Editor Changes from 21454357 are pushed
+	// First get scene outliner-style context menu
+	// FLevelEditorModule& LevelEditorModule = FModuleManager::LoadModuleChecked<FLevelEditorModule>("LevelEditor");
+	// TSharedPtr<ILevelEditor> LevelEditorInstancePtr = LevelEditorModule.GetLevelEditorInstance().Pin();
+	//
+	// UToolMenu* BaseMenu = LevelEditorInstancePtr->GenerateLevelEditorContextMenu(
+	// 		ELevelEditorMenuContext::SceneOutliner, nullptr, FTypedElementHandle());
+
+	// Now generate custom sections to append
+	UToolMenus* ToolMenus = UToolMenus::Get();
+	UToolMenu* CustomMenu = ToolMenus->GenerateMenu(DefaultContextBaseMenuName, Context);
+
+	// for (const FToolMenuSection& Section : CustomMenu->Sections)
+	// {
+	// 	BaseMenu->Sections.Add(Section);
+	// }
+
+	//return ToolMenus->GenerateWidget(BaseMenu);
+	return ToolMenus->GenerateWidget(CustomMenu);
+}
+
+void UObjectMixerEditorListMenuContext::RegisterContextMenu()
+{
+	UToolMenus* ToolMenus = UToolMenus::Get();
+	
+	if (!ToolMenus->IsMenuRegistered(DefaultContextBaseMenuName))
+	{
+		UToolMenu* Menu = ToolMenus->RegisterMenu(DefaultContextBaseMenuName);
+	
+		Menu->AddDynamicSection("DynamicCollectionsSection", FNewToolMenuDelegate::CreateLambda(
+			[](UToolMenu* InMenu)
+			{
+				UObjectMixerEditorListMenuContext* Context = InMenu->FindContext<UObjectMixerEditorListMenuContext>();
+				if (!Context || Context->Data.SelectedItems.Num() == 0)
+				{
+					return;
+				}
+	
+				FToolMenuSection& Section = InMenu->FindOrAddSection("Collections");
+				Section.Label = LOCTEXT("CollectionsSectionName", "Collections");
+			
+				Section.AddSubMenu(
+					"SelectCollectionsSubMenu",
+					LOCTEXT("SelectCollectionsSubmenu", "Select or Add Collection"),
+					LOCTEXT("SelectCollectionsSubmenu_Tooltip", "Select the collection to which you wish to assign this object."),
+					FNewToolMenuDelegate::CreateLambda(
+						[Context](UToolMenu* Menu)
+						{
+							FToolMenuEntry Args;
+							Args.Type = EMultiBlockType::Widget;
+							Args.MakeCustomWidget.BindLambda(
+								[Context](const FToolMenuContext&, const FToolMenuCustomWidgetContext&)
+								{
+									return SNew(SBox)
+											.MinDesiredWidth(200)
+											.Padding(8, 0)
+											[
+												SAssignNew(Context->EditableText, SEditableTextBox)
+												.HintText(LOCTEXT("NewCollectionEditableHintText", "Enter a new collection name..."))
+												.OnTextCommitted_Static(&UObjectMixerEditorListMenuContext::OnTextCommitted, Context)
+												.OnTextChanged_Static(&UObjectMixerEditorListMenuContext::OnTextChanged, Context)
+											]
+									;
+								}
+							);
+							Menu->AddMenuEntry("NewCollectionInput", Args);
+							
+							if (const TSharedPtr<FObjectMixerEditorMainPanel> MainPanel = Context->Data.MainPanelPtr.Pin())
+							{
+								if (TArray<FName> Collections = MainPanel->GetAllCollectionNames(); Collections.Num() > 0)
+								{
+									FToolMenuSection& Section = Menu->FindOrAddSection("Collections");
+								   Section.Label = LOCTEXT("CollectionsSectionName", "Collections");
+
+								   Collections.StableSort([](const FName& A, const FName B)
+								   {
+									   return A.LexicalLess(B);
+								   });
+								
+								   for (const FName& Key : Collections)
+								   {
+									   const FText KeyText = FText::FromName(Key);
+									
+									   Section.AddMenuEntry(
+										   Key,
+										   KeyText,
+										   FText::Format(LOCTEXT("AddObjectsToCollectionTooltipFormat", "Add selected to collection '{0}'"), KeyText),
+										   FSlateIcon(),
+										   FUIAction(
+											   FExecuteAction::CreateStatic(&UObjectMixerEditorListMenuContext::OnClickCollectionMenuEntry, Key, Context),
+											   FCanExecuteAction(),
+											   FIsActionChecked::CreateStatic(&UObjectMixerEditorListMenuContext::AreAllObjectsInCollection, Key, Context)
+										   ),
+										   EUserInterfaceActionType::Check
+									   );
+								   }
+								}
+							}
+						}
+					)
+				);
+			}
+		));
+	}
+}
+
+void UObjectMixerEditorListMenuContext::OnTextCommitted(const FText& InText, ETextCommit::Type InCommitType, UObjectMixerEditorListMenuContext* Context)
+{
+	check(Context);
+	
+	if (InCommitType == ETextCommit::OnEnter)
+	{
+		UE_LOG(LogTemp, Warning, TEXT("%hs: %s"), __FUNCTION__, *InText.ToString());
+
+		AddObjectsToCollection(*InText.ToString(), Context);
+	}
+}
+
+void UObjectMixerEditorListMenuContext::OnTextChanged(const FText& InText, UObjectMixerEditorListMenuContext* Context)
+{
+	check(Context);
+	
+	if (const FString AsString = InText.ToString(); AsString.Len() > 12)
+	{
+		Context->EditableText->SetText(FText::FromString(AsString.LeftChop(1)));
+	}
+}
+
+void UObjectMixerEditorListMenuContext::OnClickCollectionMenuEntry(const FName Key, UObjectMixerEditorListMenuContext* Context)
+{
+	if (AreAllObjectsInCollection(Key, Context))
+	{
+		RemoveObjectsFromCollection(Key, Context);
+	}
+	else
+	{
+		AddObjectsToCollection(Key, Context);
+	}
+}
+
+void UObjectMixerEditorListMenuContext::AddObjectsToCollection(const FName Key, UObjectMixerEditorListMenuContext* Context)
+{
+	check(Context);
+
+	if (const TSharedPtr<FObjectMixerEditorMainPanel> MainPanel = Context->Data.MainPanelPtr.Pin())
+	{
+		TSet<FSoftObjectPath> ObjectPaths;
+
+		for (const TSharedPtr<FObjectMixerEditorListRow>& Item : Context->Data.SelectedItems)
+		{
+			if (const UObject* Object = Item->GetObject())
+			{
+				ObjectPaths.Add(Object);
+			}
+		}
+		
+		MainPanel->AddObjectsToCollection(Key, ObjectPaths);
+	}
+}
+
+void UObjectMixerEditorListMenuContext::RemoveObjectsFromCollection(const FName Key, UObjectMixerEditorListMenuContext* Context)
+{
+	check(Context);
+
+	if (const TSharedPtr<FObjectMixerEditorMainPanel> MainPanel = Context->Data.MainPanelPtr.Pin())
+	{
+		TSet<FSoftObjectPath> ObjectPaths;
+
+		for (const TSharedPtr<FObjectMixerEditorListRow>& Item : Context->Data.SelectedItems)
+		{
+			if (const UObject* Object = Item->GetObject())
+			{
+				ObjectPaths.Add(Object);
+			}
+		}
+		
+		MainPanel->RemoveObjectsFromCollection(Key, ObjectPaths);
+	}
+}
+
+bool UObjectMixerEditorListMenuContext::AreAllObjectsInCollection(const FName Key, UObjectMixerEditorListMenuContext* Context)
+{
+	check(Context);
+
+	bool bAreAllSelectedObjectsInCollection = false;
+
+	if (const TSharedPtr<FObjectMixerEditorMainPanel> MainPanel = Context->Data.MainPanelPtr.Pin())
+	{
+		for (const TSharedPtr<FObjectMixerEditorListRow>& Item : Context->Data.SelectedItems)
+		{
+			if (const UObject* Object = Item->GetObject())
+			{
+				bAreAllSelectedObjectsInCollection = MainPanel->IsObjectInCollection(Key, Object);
+
+				if (!bAreAllSelectedObjectsInCollection)
+				{
+					return false;
+				}
+			}
+		}
+	}
+
+	return bAreAllSelectedObjectsInCollection;
+}
+
+#undef LOCTEXT_NAMESPACE
diff --git a/Engine/Plugins/Editor/ObjectMixer/ObjectMixer/Source/ObjectMixer/Private/Views/Widgets/SCollectionSelectionButton.cpp b/Engine/Plugins/Editor/ObjectMixer/ObjectMixer/Source/ObjectMixer/Private/Views/Widgets/SCollectionSelectionButton.cpp
new file mode 100644
--- /dev/null
+++ b/Engine/Plugins/Editor/ObjectMixer/ObjectMixer/Source/ObjectMixer/Private/Views/Widgets/SCollectionSelectionButton.cpp
@@ -0,0 +1,238 @@
+// Copyright Epic Games, Inc. All Rights Reserved.
+
+#pragma once
+
+#include "Views/Widgets/SCollectionSelectionButton.h"
+
+#include "DragAndDrop/DecoratedDragDropOp.h"
+#include "Styling/AppStyle.h"
+#include "Views/MainPanel/ObjectMixerEditorMainPanel.h"
+#include "Views/MainPanel/SObjectMixerEditorMainPanel.h"
+#include "Widgets/Layout/SBorder.h"
+#include "Widgets/Text/STextBlock.h"
+
+#define LOCTEXT_NAMESPACE "ObjectMixerEditor"
+
+class FCollectionSelectionButtonDragDropOp : public FDecoratedDragDropOp
+{
+public:
+	DRAG_DROP_OPERATOR_TYPE(FCollectionSelectionButtonDragDropOp, FDecoratedDragDropOp)
+
+	/** The item being dragged and dropped */
+	FName DraggedItem;
+
+	/** Constructs a new drag/drop operation */
+	static TSharedRef<FCollectionSelectionButtonDragDropOp> New(const FName DraggedItem)
+	{
+		TSharedRef<FCollectionSelectionButtonDragDropOp> Operation = MakeShareable(
+			   new FCollectionSelectionButtonDragDropOp());
+		if (DraggedItem != NAME_None)
+		{
+			Operation->DraggedItem = DraggedItem;
+
+			Operation->DefaultHoverIcon = FAppStyle::Get().GetBrush("Graph.ConnectorFeedback.Error");
+
+			Operation->DefaultHoverText = LOCTEXT("DefaultCollectionButtonHoverText","Drop onto another Collection Button to set a custom order.");
+
+			Operation->Construct();
+		}
+
+		return Operation;
+	}
+};
+
+void SCollectionSelectionButton::Construct(const FArguments& InArgs, const TSharedRef<SObjectMixerEditorMainPanel> MainPanelWidget, const FName& InCollectionName)
+{
+	MainPanelPtr = MainPanelWidget;
+	CollectionName = InCollectionName;
+	
+	ChildSlot
+	[
+		SNew(SBorder)
+		.Padding(FMargin(16, 4))
+		.BorderImage_Lambda([this]()
+		{
+			if (MainPanelPtr.Pin()->IsCollectionChecked(CollectionName) == ECheckBoxState::Checked)
+			{
+				if (bIsPressed)
+				{
+					return &CheckedPressedImage;
+				}
+				else if (bDropIsValid)
+				{
+					return &CheckedValidDropImage;
+				}
+				else if (IsHovered())
+				{
+					return &CheckedHoveredImage;
+				}
+
+				return &CheckedImage;
+			}
+
+			if (bIsPressed)
+			{
+				return &UncheckedPressedImage;
+			}
+			else if (bDropIsValid)
+			{
+				return &UncheckedValidDropImage;
+			}
+			else if (IsHovered())
+			{
+				return &UncheckedHoveredImage;
+			}
+
+			return &UncheckedImage;
+			
+		})
+		[
+			SNew(STextBlock)
+			.TextStyle(FAppStyle::Get(), "SmallText")
+			.Text(FText::FromName(CollectionName))
+			.Visibility(EVisibility::HitTestInvisible)
+		]
+	];
+}
+
+FReply SCollectionSelectionButton::OnMouseButtonDown(const FGeometry& MyGeometry, const FPointerEvent& MouseEvent) 
+{
+	bIsPressed = true;
+	if (MouseEvent.GetEffectingButton() == EKeys::LeftMouseButton )
+	{
+		return FReply::Handled().DetectDrag( SharedThis(this), EKeys::LeftMouseButton );
+	}
+	return FReply::Handled();
+}
+
+FReply SCollectionSelectionButton::OnMouseButtonUp(const FGeometry& MyGeometry, const FPointerEvent& MouseEvent) 
+{
+	bIsPressed = false;
+	if (MouseEvent.GetEffectingButton() == EKeys::LeftMouseButton )
+	{
+		ECheckBoxState CurrentState = MainPanelPtr.Pin()->IsCollectionChecked(CollectionName);
+
+		ECheckBoxState NewState = ECheckBoxState::Checked;
+
+		switch (CurrentState)
+		{
+		case ECheckBoxState::Checked:
+			NewState = ECheckBoxState::Unchecked;
+			break;
+
+		default:
+			break;
+		}
+		
+		MainPanelPtr.Pin()->OnCollectionCheckedStateChanged(NewState, CollectionName);
+	}
+	if (MouseEvent.GetEffectingButton() == EKeys::RightMouseButton)
+	{
+		MainPanelPtr.Pin()->RequestRemoveCollection(CollectionName);
+	}
+	return FReply::Handled();
+}
+
+FReply SCollectionSelectionButton::OnDragDetected(const FGeometry& MyGeometry, const FPointerEvent& MouseEvent)
+{
+	if (CollectionName != "All")
+	{
+		TSharedRef<FCollectionSelectionButtonDragDropOp> OperationFromCollection =
+		   FCollectionSelectionButtonDragDropOp::New(CollectionName);
+
+		OperationFromCollection->ResetToDefaultToolTip();
+
+		bDropIsValid = false;
+
+		return FReply::Handled().BeginDragDrop(OperationFromCollection);
+	}
+
+	return FReply::Handled();
+}
+
+void SCollectionSelectionButton::OnDragEnter(const FGeometry& MyGeometry, const FDragDropEvent& DragDropEvent)
+{
+	if (TSharedPtr<FObjectMixerListRowDragDropOp> OperationFromRow =
+		DragDropEvent.GetOperationAs<FObjectMixerListRowDragDropOp>())
+	{
+		if (CollectionName != "All" && CollectionName != NAME_None)
+		{
+			OperationFromRow->SetToolTip(
+			   LOCTEXT("DropRowItemsOntoCollectionButtonCTA","Add selected items to this collection"),
+			   FAppStyle::Get().GetBrush("Graph.ConnectorFeedback.OK"));
+
+			bDropIsValid = true;
+		}
+	}
+
+	if (TSharedPtr<FCollectionSelectionButtonDragDropOp> OperationFromCollection =
+		DragDropEvent.GetOperationAs<FCollectionSelectionButtonDragDropOp>())
+	{
+		if (CollectionName != NAME_None && CollectionName != OperationFromCollection->DraggedItem)
+		{
+			OperationFromCollection->SetToolTip(
+			   FText::Format(
+				LOCTEXT("DropCollectionButtonOntoCollectionButtonCTA_Format","Reorder {0} before {1}"),
+				FText::FromName(OperationFromCollection->DraggedItem), FText::FromName(CollectionName)),
+			   FAppStyle::Get().GetBrush("Graph.ConnectorFeedback.OK"));
+
+			bDropIsValid = true;
+		}
+	}
+}
+
+void SCollectionSelectionButton::OnDragLeave(const FDragDropEvent& DragDropEvent)
+{
+	bIsPressed = false;
+	bDropIsValid = false;
+	
+	if (TSharedPtr<FObjectMixerListRowDragDropOp> OperationFromRow =
+		DragDropEvent.GetOperationAs<FObjectMixerListRowDragDropOp>())
+	{
+		OperationFromRow->ResetToDefaultToolTip();
+	}
+
+	if (TSharedPtr<FCollectionSelectionButtonDragDropOp> OperationFromCollection =
+		DragDropEvent.GetOperationAs<FCollectionSelectionButtonDragDropOp>())
+	{
+		OperationFromCollection->ResetToDefaultToolTip();
+	}
+}
+
+FReply SCollectionSelectionButton::OnDrop(const FGeometry& MyGeometry, const FDragDropEvent& DragDropEvent)
+{
+	bIsPressed = false;
+	bDropIsValid = false;
+	
+	if (TSharedPtr<FObjectMixerListRowDragDropOp> OperationFromRow =
+		DragDropEvent.GetOperationAs<FObjectMixerListRowDragDropOp>())
+	{
+		if (const TSharedPtr<FObjectMixerEditorMainPanel> MainPanelModel = MainPanelPtr.Pin()->GetMainPanelModel().Pin())
+		{
+			TSet<FSoftObjectPath> ObjectPaths;
+
+			for (const TSharedPtr<FObjectMixerEditorListRow>& Item : OperationFromRow->DraggedItems)
+			{
+				if (const UObject* Object = Item->GetObject())
+				{
+					ObjectPaths.Add(Object);
+				}
+			}
+	
+			MainPanelModel->AddObjectsToCollection(CollectionName, ObjectPaths);
+		}
+	}
+
+	if (TSharedPtr<FCollectionSelectionButtonDragDropOp> OperationFromCollection =
+		DragDropEvent.GetOperationAs<FCollectionSelectionButtonDragDropOp>())
+	{
+		if (const TSharedPtr<FObjectMixerEditorMainPanel> MainPanelModel = MainPanelPtr.Pin()->GetMainPanelModel().Pin())
+		{
+			MainPanelModel->ReorderCollection(OperationFromCollection->DraggedItem, CollectionName);
+		}
+	}
+		
+	return FReply::Handled();
+}
+
+#undef LOCTEXT_NAMESPACE
diff --git a/Engine/Plugins/Editor/ObjectMixer/ObjectMixer/Source/ObjectMixer/Private/Views/Widgets/SObjectMixerPlacementAssetMenuEntry.cpp b/Engine/Plugins/Editor/ObjectMixer/ObjectMixer/Source/ObjectMixer/Private/Views/Widgets/SObjectMixerPlacementAssetMenuEntry.cpp
new file mode 100644
--- /dev/null
+++ b/Engine/Plugins/Editor/ObjectMixer/ObjectMixer/Source/ObjectMixer/Private/Views/Widgets/SObjectMixerPlacementAssetMenuEntry.cpp
@@ -0,0 +1,252 @@
+// Copyright Epic Games, Inc. All Rights Reserved.
+
+#include "SObjectMixerPlacementAssetMenuEntry.h"
+
+#include "AssetSelection.h"
+#include "ClassIconFinder.h"
+#include "DragAndDrop/AssetDragDropOp.h"
+#include "EditorClassUtils.h"
+#include "LevelEditorViewport.h"
+#include "PlacementMode/Public/IPlacementModeModule.h"
+#include "Styling/SlateIconFinder.h"
+#include "Styling/StyleColors.h"
+#include "Widgets/Layout/SSpacer.h"
+
+void SObjectMixerPlacementAssetMenuEntry::Construct(const FArguments& InArgs, const TSharedPtr<const FPlaceableItem>& InItem)
+{	
+	bIsPressed = false;
+
+	check(InItem.IsValid());
+
+	Item = InItem;
+
+	AssetImage = nullptr;
+
+	TSharedPtr< SHorizontalBox > ActorType = SNew( SHorizontalBox );
+
+	const bool bIsClass = Item->AssetData.GetClass() == UClass::StaticClass();
+	const bool bIsActor = bIsClass ? CastChecked<UClass>(Item->AssetData.GetAsset())->IsChildOf(AActor::StaticClass()) : false;
+
+	AActor* DefaultActor = nullptr;
+	if (Item->Factory != nullptr)
+	{
+		DefaultActor = Item->Factory->GetDefaultActor(Item->AssetData);
+	}
+	else if (bIsActor)
+	{
+		DefaultActor = CastChecked<AActor>(CastChecked<UClass>(Item->AssetData.GetAsset())->ClassDefaultObject);
+	}
+
+	UClass* DocClass = nullptr;
+	TSharedPtr<IToolTip> AssetEntryToolTip;
+	if(DefaultActor != nullptr)
+	{
+		DocClass = DefaultActor->GetClass();
+		AssetEntryToolTip = FEditorClassUtils::GetTooltip(DefaultActor->GetClass());
+	}
+
+	if (!AssetEntryToolTip.IsValid())
+	{
+		AssetEntryToolTip = FSlateApplicationBase::Get().MakeToolTip(Item->DisplayName);
+	}
+	
+	const FButtonStyle& ButtonStyle = FAppStyle::Get().GetWidgetStyle<FButtonStyle>( "Menu.Button" );
+	const float MenuIconSize = FAppStyle::Get().GetFloat("Menu.MenuIconSize");
+
+	Style = &ButtonStyle;
+
+	// Create doc link widget if there is a class to link to
+	TSharedRef<SWidget> DocWidget = SNew(SSpacer);
+	if(DocClass != NULL)
+	{
+		DocWidget = FEditorClassUtils::GetDocumentationLinkWidget(DocClass);
+		DocWidget->SetCursor( EMouseCursor::Default );
+	}
+
+	ChildSlot
+	.HAlign(HAlign_Fill)
+	.VAlign(VAlign_Fill)
+	[
+		SNew(SBorder)
+		.BorderImage( this, &SObjectMixerPlacementAssetMenuEntry::GetBorder )
+		.Cursor( EMouseCursor::GrabHand )
+		.ToolTip( AssetEntryToolTip )
+		.Padding(FMargin(5.f, 3.f, 5.f, 3.f))
+		[
+			SNew( SHorizontalBox )
+
+			+ SHorizontalBox::Slot()
+			.Padding(10.0f, 0.f, 10.f, 0.0f)
+			.AutoWidth()
+			.HAlign(HAlign_Left)
+			.VAlign(VAlign_Center)
+			[
+				SNew(SBox)
+				.WidthOverride(MenuIconSize)
+				.HeightOverride(MenuIconSize)
+				[
+					SNew(SImage)
+					.Image(this, &SObjectMixerPlacementAssetMenuEntry::GetIcon)
+					.ColorAndOpacity(FSlateColor::UseSubduedForeground())
+				]
+			]
+
+			+ SHorizontalBox::Slot()
+			.FillWidth(1.f)
+			.Padding(1.f, 0.f, 0.f, 0.f)
+			.VAlign(VAlign_Center)
+			.HAlign(HAlign_Left)
+			[
+
+				SNew( STextBlock )
+				.ColorAndOpacity(FSlateColor::UseForeground())
+				.Text( Item->DisplayName )
+			]
+
+
+			+ SHorizontalBox::Slot()
+			.VAlign(VAlign_Center)
+			.HAlign(HAlign_Right)
+			.AutoWidth()
+			[
+				SNew(SImage)
+                .ColorAndOpacity(FSlateColor::UseSubduedForeground())
+				.Image(FAppStyle::Get().GetBrush("Icons.DragHandle"))
+			]
+		]
+	];
+}
+
+
+const FSlateBrush* SObjectMixerPlacementAssetMenuEntry::GetIcon() const
+{
+
+	if (AssetImage != nullptr)
+	{
+		return AssetImage;
+	}
+
+	if (Item->ClassIconBrushOverride != NAME_None)
+	{
+		AssetImage = FSlateIconFinder::FindCustomIconBrushForClass(nullptr, TEXT("ClassIcon"), Item->ClassIconBrushOverride);
+	}
+	else
+	{
+		AssetImage = FSlateIconFinder::FindIconBrushForClass(FClassIconFinder::GetIconClassForAssetData(Item->AssetData));
+	}
+
+	return AssetImage;
+}
+
+
+FReply SObjectMixerPlacementAssetMenuEntry::OnMouseButtonDown(const FGeometry& MyGeometry, const FPointerEvent& MouseEvent)
+{
+	if ( MouseEvent.GetEffectingButton() == EKeys::LeftMouseButton )
+	{
+		bIsPressed = true;
+
+		return FReply::Handled().DetectDrag( SharedThis( this ), MouseEvent.GetEffectingButton() );
+	}
+
+	return FReply::Unhandled();
+}
+
+FReply SObjectMixerPlacementAssetMenuEntry::OnMouseButtonUp(const FGeometry& MyGeometry, const FPointerEvent& MouseEvent)
+{
+	if ( MouseEvent.GetEffectingButton() == EKeys::LeftMouseButton )
+	{
+		bIsPressed = false;
+
+		AActor* NewActor = nullptr;
+		UActorFactory* Factory = Item->Factory;
+		if (!Item->Factory)
+		{
+			// If no actor factory was found or failed, add the actor from the uclass
+			UClass* AssetClass = Item->AssetData.GetClass();
+			if (AssetClass)
+			{
+				UObject* ClassObject = AssetClass->GetDefaultObject();
+				FActorFactoryAssetProxy::GetFactoryForAssetObject(ClassObject);
+			}
+		}
+		NewActor = GEditor->GetEditorWorldContext().World()->SpawnActor(Factory->NewActorClass);
+		if (NewActor && GCurrentLevelEditingViewportClient)
+		{
+  			GEditor->MoveActorInFrontOfCamera(*NewActor, 
+  				GCurrentLevelEditingViewportClient->GetViewLocation(), 
+  				GCurrentLevelEditingViewportClient->GetViewRotation().Vector()
+  			);
+		}
+
+		if (!MouseEvent.IsControlDown())
+		{
+			FSlateApplication::Get().DismissAllMenus();
+		}
+
+		GEditor->SelectNone(true, true, false);
+		GEditor->SelectActor(NewActor, true, false);
+
+		return FReply::Handled();
+	}
+
+	return FReply::Unhandled();
+}
+
+FReply SObjectMixerPlacementAssetMenuEntry::OnDragDetected(const FGeometry& MyGeometry, const FPointerEvent& MouseEvent)
+{
+	bIsPressed = false;
+
+	if (FEditorDelegates::OnAssetDragStarted.IsBound())
+	{
+		TArray<FAssetData> DraggedAssetDatas;
+		DraggedAssetDatas.Add( Item->AssetData );
+		FEditorDelegates::OnAssetDragStarted.Broadcast( DraggedAssetDatas, Item->Factory );
+		return FReply::Handled();
+	}
+
+	if( MouseEvent.IsMouseButtonDown( EKeys::LeftMouseButton ) )
+	{
+		return FReply::Handled().BeginDragDrop(FAssetDragDropOp::New(Item->AssetData, Item->Factory));
+	}
+	else
+	{
+		return FReply::Handled();
+	}
+}
+
+bool SObjectMixerPlacementAssetMenuEntry::IsPressed() const
+{
+	return bIsPressed;
+}
+
+const FSlateBrush* SObjectMixerPlacementAssetMenuEntry::GetBorder() const
+{
+	if ( IsPressed() )
+	{
+		return &(Style->Pressed);
+	}
+	else if ( IsHovered() )
+	{
+		return &(Style->Hovered);
+	}
+	else
+	{
+		return &(Style->Normal);
+	}
+}
+
+FSlateColor SObjectMixerPlacementAssetMenuEntry::GetForegroundColor() const
+{
+	if (IsPressed())
+	{
+		return Style->PressedForeground;
+	}
+	else if (IsHovered())
+	{
+		return Style->HoveredForeground;
+	}
+	else
+	{
+		return Style->NormalForeground;
+	}
+}
diff --git a/Engine/Plugins/Editor/ObjectMixer/ObjectMixer/Source/ObjectMixer/Private/Views/Widgets/SObjectMixerPlacementAssetMenuEntry.h b/Engine/Plugins/Editor/ObjectMixer/ObjectMixer/Source/ObjectMixer/Private/Views/Widgets/SObjectMixerPlacementAssetMenuEntry.h
new file mode 100644
--- /dev/null
+++ b/Engine/Plugins/Editor/ObjectMixer/ObjectMixer/Source/ObjectMixer/Private/Views/Widgets/SObjectMixerPlacementAssetMenuEntry.h
@@ -0,0 +1,36 @@
+// Copyright Epic Games, Inc. All Rights Reserved.
+
+#pragma once
+
+#include "Styling/SlateBrush.h"
+#include "Styling/SlateTypes.h"
+#include "Widgets/SCompoundWidget.h"
+
+class SObjectMixerPlacementAssetMenuEntry : public SCompoundWidget
+{
+public:
+	SLATE_BEGIN_ARGS(SObjectMixerPlacementAssetMenuEntry){}
+	SLATE_END_ARGS()
+
+	void Construct(const FArguments& InArgs, const TSharedPtr<const struct FPlaceableItem>& InItem);
+
+	virtual FReply OnMouseButtonDown(const FGeometry& MyGeometry, const FPointerEvent& MouseEvent) override;
+	virtual FReply OnMouseButtonUp(const FGeometry& MyGeometry, const FPointerEvent& MouseEvent) override;
+	virtual FReply OnDragDetected(const FGeometry& MyGeometry, const FPointerEvent& MouseEvent) override;
+
+	bool IsPressed() const;
+
+	TSharedPtr<const FPlaceableItem> Item;
+
+	virtual FSlateColor GetForegroundColor() const override;
+
+private:
+	const FSlateBrush* GetBorder() const;
+	const FSlateBrush* GetIcon() const;
+
+	bool bIsPressed = false;
+
+	const FButtonStyle* Style = nullptr;
+	
+	mutable const FSlateBrush* AssetImage = nullptr;
+};
diff --git a/Engine/Plugins/Editor/ObjectMixer/ObjectMixer/Source/ObjectMixer/Public/ObjectFilter/ObjectMixerEditorObjectFilter.h b/Engine/Plugins/Editor/ObjectMixer/ObjectMixer/Source/ObjectMixer/Public/ObjectFilter/ObjectMixerEditorObjectFilter.h
new file mode 100644
--- /dev/null
+++ b/Engine/Plugins/Editor/ObjectMixer/ObjectMixer/Source/ObjectMixer/Public/ObjectFilter/ObjectMixerEditorObjectFilter.h
@@ -0,0 +1,312 @@
+// Copyright Epic Games, Inc. All Rights Reserved.
+
+#pragma once
+
+#include "Components/LightComponent.h"
+#include "UObject/Class.h"
+
+#include "ObjectMixerEditorObjectFilter.generated.h"
+
+UENUM(BlueprintType)
+enum class EObjectMixerInheritanceInclusionOptions : uint8
+{
+	None, // Get only the properties in the specified classes without considering parent or child classes + Specified Class
+	IncludeOnlyImmediateParent, // Get properties from the class the specified class immediately derives from, but not their parents + Specified Class
+	IncludeOnlyImmediateChildren, // Get properties from child classes but not child classes of child classes + Specified Class
+	IncludeOnlyImmediateParentAndChildren, // IncludeOnlyImmediateParent + IncludeOnlyImmediateChildren + Specified Class
+	IncludeAllParents, // Go up the chain of Super classes to get all properties in the class' ancestry + Specified Class
+	IncludeAllChildren, // Get properties from all derived classes recursively + Specified Class
+	IncludeAllParentsAndChildren, // IncludeAllParents + IncludeAllChildren + Specified Class
+	IncludeAllParentsAndOnlyImmediateChildren, // IncludeAllParents + IncludeOnlyImmediateChildren + Specified Class
+	IncludeOnlyImmediateParentAndAllChildren // IncludeOnlyImmediateParent + IncludeAllChildren + Specified Class
+};
+
+UENUM(BlueprintType)
+enum class EObjectMixerTreeViewMode : uint8
+{
+	Flat, // Show all matching objects in a flat list without any hierarchy
+	Folder, // Display objects in a folder hierarchy but without Object->SubObject hierarchy
+	ObjectSubObject, // Display objects as SubObjects of BaseObjects (usually actors) without considering folder hierarchy
+	FolderObjectSubObject // Display objects in a full hierarchy (folder->BaseObject->SubObject)
+};
+
+/**
+ * Native class for filtering object types to Object Mixer.
+ * Native C++ classes should inherit directly from this class.
+ */
+UCLASS(Abstract, Blueprintable)
+class OBJECTMIXEREDITOR_API UObjectMixerObjectFilter : public UObject
+{
+	GENERATED_BODY()
+	
+public:
+
+	UObjectMixerObjectFilter() = default;
+	
+	/**
+	 * Return the basic object types you want to filter for in your level.
+	 * For example, if you want to work with Lights, return ULightComponentBase.
+	 * If you also want to see the properties for parent or child classes,
+	 * override the GetObjectMixerPropertyInheritanceInclusionOptions and GetForceAddedColumns functions.
+	 */
+	virtual TSet<UClass*> GetObjectClassesToFilter() const { return {}; }
+
+	/**
+	 * Return the basic actor types you want to be able to place using the Add button.
+	 * Note that only subclasses of AActor are supported and only those which have a registered factory.
+	 * This includes most engine actor types.
+	 */
+	virtual TSet<TSubclassOf<AActor>> GetObjectClassesToPlace() const { return {}; }
+
+	/**
+	 * Get the text to display for the object name/label.
+	 * This is useful if one of your classes is a component type and you want the label of the component's owning actor, for example.
+	 * If not overridden, this returns the object's name.
+	 * @param InViewMode The tree view display option (i.e. Flat list or hierarchy).
+	 */
+	virtual FText GetRowDisplayName(
+		UObject* InObject, EObjectMixerTreeViewMode InViewMode) const;
+
+	/**
+	 * Controls how to display the row's visibility icon. Return true if the object should be visible.
+	 * Generally this should work like the Scene Outliner does.
+	 * If not overridden, we use the Editor Visibility of the object's AActor outer (unless it's an actor itself).
+	 * @param InViewMode The tree view display option (i.e. Flat list or hierarchy) (unused)
+	 */
+	virtual bool GetRowEditorVisibility(
+		UObject* InObject, EObjectMixerTreeViewMode InViewMode) const;
+
+	/**
+	 * Controls what happens when the row's visibility icon is clicked.
+	 * Generally this should work like the Scene Outliner does.
+	 * If not overridden, we set the Editor Visibility of the object's AActor outer (unless it's an actor itself).
+	 * @param InViewMode The tree view display option (i.e. Flat list or hierarchy) (unused)
+	 */
+	virtual void OnSetRowEditorVisibility(
+		UObject* InObject, bool bNewIsVisible, EObjectMixerTreeViewMode InViewMode) const;
+
+	/**
+	 * Specify a list of property names corresponding to columns you want to show by default.
+	 * For example, you can specify "Intensity" and "LightColor" to show only those property columns by default in the UI.
+	 * Columns not specified will not be shown by default but can be enabled by the user in the UI.
+	 */
+	virtual TSet<FName> GetColumnsToShowByDefault() const;
+
+	/**
+	 * Specify a list of property names corresponding to columns you don't want to ever show.
+	 * For example, you can specify "Intensity" and "LightColor" to ensure that they can't be enabled or shown in the UI.
+	 * Columns not specified can be enabled by the user in the UI.
+	 */
+	virtual TSet<FName> GetColumnsToExclude() const;
+
+	/**
+	 * Specify a list of property names found in parent classes you want to show that aren't in the specified classes.
+	 * Note that properties specified here do not override the properties specified in GetColumnsToExclude(),
+	 * but do override the supported property tests so these will appear even if ShouldIncludeUnsupportedProperties returns false.
+	 * For example, a ULightComponent displays "LightColor" in the editor's details panel,
+	 * but ULightComponent itself doesn't have a property named "LightColor". Instead it's in its parent class, ULightComponentBase.
+	 * In this scenario, ULightComponent is specified and PropertyInheritanceInclusionOptions is None, so "LightColor" won't appear by default.
+	 * Specify "LightColor" in this function to ensure that "LightColor" will appear as a column as long as
+	 * the property is accessible to one of the specified classes regardless of which parent class it comes from.
+	 */
+	virtual TSet<FName> GetForceAddedColumns() const { return {}; }
+
+	/**
+	 * Specify whether we should return only the properties of the specified classes or the properties of parent and child classes.
+	 * Defaults to 'None' which only considers the properties of the specified classes.
+	 * If you're not seeing all the properties you expected, try overloading this function.
+	 */
+	virtual EObjectMixerInheritanceInclusionOptions GetObjectMixerPropertyInheritanceInclusionOptions() const;
+
+	/**
+	 * Specify whether we should return only the specified classes or the parent and child classes in placement mode.
+	 * Defaults to 'None' which only considers the specified classes.
+	 */
+	virtual EObjectMixerInheritanceInclusionOptions GetObjectMixerPlacementClassInclusionOptions() const;
+
+	/**
+	 * If true, properties that are not visible in the details panel and properties not supported by SSingleProperty will be selectable.
+	 * Defaults to false.
+	 */
+	virtual bool ShouldIncludeUnsupportedProperties() const;
+
+	static TSet<UClass*> GetParentAndChildClassesFromSpecifiedClasses(
+		const TSet<UClass*>& InSpecifiedClasses, EObjectMixerInheritanceInclusionOptions Options);
+
+	static TSet<UClass*> GetParentAndChildClassesFromSpecifiedClasses(
+		const TSet<TSubclassOf<AActor>>& InSpecifiedClasses, EObjectMixerInheritanceInclusionOptions Options);
+
+	static EFieldIterationFlags GetDesiredFieldIterationFlags(const bool bIncludeInheritedProperties);
+
+protected:
+
+	/**
+	 * Given a set of property names you wish to include, returns a list of all other properties on InObject not found in ExcludeList.
+	 * Useful when defining default visible columns in a list view.
+	 */
+	TSet<FName> GenerateIncludeListFromExcludeList(const TSet<FName>& ExcludeList) const;
+};
+
+/**
+ * Script class for filtering object types to Object Mixer.
+ * Blueprint classes should inherit directly from this class.
+ */
+UCLASS(Abstract, Blueprintable)
+class OBJECTMIXEREDITOR_API UObjectMixerBlueprintObjectFilter : public UObjectMixerObjectFilter
+{
+	GENERATED_BODY()
+public:
+	
+	/**
+	 * Return the basic object types you want to filter for in your level.
+	 * For example, if you want to work with Lights, return ULightComponentBase.
+	 * If you also want to see the properties for parent or child classes,
+	 * override the GetObjectMixerPropertyInheritanceInclusionOptions and GetForceAddedColumns functions.
+	 */
+	UFUNCTION(BlueprintNativeEvent, BlueprintCallable, Category = "Object Mixer")
+	TSet<UClass*> GetObjectClassesToFilter() const override;
+
+	TSet<UClass*> GetObjectClassesToFilter_Implementation() const
+	{
+		return Super::GetObjectClassesToFilter();
+	}
+
+	/**
+	 * Return the basic actor types you want to be able to place using the Add button.
+	 * Note that only subclasses of AActor are supported and only those which have a registered factory.
+	 * This includes most engine actor types.
+	 */
+	UFUNCTION(BlueprintNativeEvent, BlueprintCallable, Category = "Object Mixer")
+	TSet<TSubclassOf<AActor>> GetObjectClassesToPlace() const override;
+
+	TSet<TSubclassOf<AActor>> GetObjectClassesToPlace_Implementation() const
+	{
+		return Super::GetObjectClassesToPlace();
+	}
+
+	/**
+	 * Get the text to display for the object name/label.
+	 * This is useful if one of your classes is a component type and you want the label of the component's owning actor, for example.
+	 * If not overridden, this returns the object's name.
+	 * @param InViewMode The tree view display option (i.e. Flat list or hierarchy).
+	 */
+	UFUNCTION(BlueprintNativeEvent, BlueprintCallable, Category = "Object Mixer")
+	FText GetRowDisplayName(
+		UObject* InObject, EObjectMixerTreeViewMode InViewMode = EObjectMixerTreeViewMode::FolderObjectSubObject) const override;
+	
+	FText GetRowDisplayName_Implementation(UObject* InObject, EObjectMixerTreeViewMode InViewMode) const
+	{
+		return Super::GetRowDisplayName(InObject, InViewMode);
+	}
+	
+	/**
+	 * Controls how to display the row's visibility icon. Return true if the object should be visible.
+	 * Generally this should work like the Scene Outliner does.
+	 * If not overridden, we use the Editor Visibility of the object's AActor outer (unless it's an actor itself).
+	 * @param InViewMode The tree view display option (i.e. Flat list or hierarchy) (unused)
+	 */
+	UFUNCTION(BlueprintNativeEvent, BlueprintCallable, Category = "Object Mixer")
+	bool GetRowEditorVisibility(
+		UObject* InObject, EObjectMixerTreeViewMode InViewMode = EObjectMixerTreeViewMode::FolderObjectSubObject) const override;
+
+	bool GetRowEditorVisibility_Implementation(UObject* InObject, EObjectMixerTreeViewMode InViewMode) const
+	{
+		return Super::GetRowEditorVisibility(InObject, InViewMode);
+	}
+	
+	/**
+	 * Controls what happens when the row's visibility icon is clicked.
+	 * Generally this should work like the Scene Outliner does.
+	 * If not overridden, we set the Editor Visibility of the object's AActor outer (unless it's an actor itself).
+	 * @param InViewMode The tree view display option (i.e. Flat list or hierarchy) (unused)
+	 */
+	UFUNCTION(BlueprintNativeEvent, BlueprintCallable, Category = "Object Mixer")
+	void OnSetRowEditorVisibility(
+		UObject* InObject, bool bNewIsVisible, EObjectMixerTreeViewMode InViewMode = EObjectMixerTreeViewMode::FolderObjectSubObject) const override;
+
+	void OnSetRowEditorVisibility_Implementation(UObject* InObject, bool bNewIsVisible, EObjectMixerTreeViewMode InViewMode) const
+	{
+		return Super::OnSetRowEditorVisibility(InObject, bNewIsVisible, InViewMode);
+	}
+	
+	/**
+	 * Specify a list of property names corresponding to columns you want to show by default.
+	 * For example, you can specify "Intensity" and "LightColor" to show only those property columns by default in the UI.
+	 * Columns not specified will not be shown by default but can be enabled by the user in the UI.
+	 */
+	UFUNCTION(BlueprintNativeEvent, BlueprintCallable, Category = "Object Mixer")
+	TSet<FName> GetColumnsToShowByDefault() const override;
+	
+	TSet<FName> GetColumnsToShowByDefault_Implementation() const
+	{
+		return Super::GetColumnsToShowByDefault();
+	}
+	
+	/**
+	 * Specify a list of property names corresponding to columns you don't want to ever show.
+	 * For example, you can specify "Intensity" and "LightColor" to ensure that they can't be enabled or shown in the UI.
+	 * Columns not specified can be enabled by the user in the UI.
+	 */
+	UFUNCTION(BlueprintNativeEvent, BlueprintCallable, Category = "Object Mixer")
+	TSet<FName> GetColumnsToExclude() const override;
+
+	TSet<FName> GetColumnsToExclude_Implementation() const
+	{
+		return Super::GetColumnsToExclude();
+	}
+
+	/**
+	 * Specify a list of property names found in parent classes you want to show that aren't in the specified classes.
+	 * Note that properties specified here do not override the properties specified in GetColumnsToExclude().
+	 * For example, a ULightComponent displays "LightColor" in the editor's details panel,
+	 * but ULightComponent itself doesn't have a property named "LightColor". Instead it's in its parent class, ULightComponentBase.
+	 * In this scenario, ULightComponent is specified and PropertyInheritanceInclusionOptions is None, so "LightColor" won't appear by default.
+	 * Specify "LightColor" in this function to ensure that "LightColor" will appear as a column as long as
+	 * the property is accessible to one of the specified classes regardless of which parent class it comes from.
+	 */
+	UFUNCTION(BlueprintNativeEvent, BlueprintCallable, Category = "Object Mixer")
+	TSet<FName> GetForceAddedColumns() const override;
+	
+	TSet<FName> GetForceAddedColumns_Implementation() const
+	{
+		return Super::GetForceAddedColumns();
+	}
+	
+	/**
+	 * Specify whether we should return only the properties of the specified classes or the properties of parent and child classes.
+	 * Defaults to 'None' which only considers the properties of the specified classes.
+	 * If you're not seeing all the properties you expected, try overloading this function.
+	 */
+	UFUNCTION(BlueprintNativeEvent, BlueprintCallable, Category = "Object Mixer")
+	EObjectMixerInheritanceInclusionOptions GetObjectMixerPropertyInheritanceInclusionOptions() const override;
+
+	EObjectMixerInheritanceInclusionOptions GetObjectMixerPropertyInheritanceInclusionOptions_Implementation() const
+	{
+		return Super::GetObjectMixerPropertyInheritanceInclusionOptions();
+	}
+
+	/**
+	 * Specify whether we should return only the specified classes or the parent and child classes in placement mode.
+	 * Defaults to 'None' which only considers the specified classes.
+	 */
+	UFUNCTION(BlueprintNativeEvent, BlueprintCallable, Category = "Object Mixer")
+	EObjectMixerInheritanceInclusionOptions GetObjectMixerPlacementClassInclusionOptions() const override;
+
+	EObjectMixerInheritanceInclusionOptions GetObjectMixerPlacementClassInclusionOptions_Implementation() const
+	{
+		return Super::GetObjectMixerPlacementClassInclusionOptions();
+	}
+
+	/**
+	 * If true, properties that are not visible in the details panel and properties not supported by SSingleProperty will be selectable.
+	 * Defaults to false.
+	 */
+	UFUNCTION(BlueprintNativeEvent, BlueprintCallable, Category = "Object Mixer")
+	bool ShouldIncludeUnsupportedProperties() const override;
+
+	bool ShouldIncludeUnsupportedProperties_Implementation() const
+	{
+		return Super::ShouldIncludeUnsupportedProperties();
+	}
+
+};
diff --git a/Engine/Plugins/Editor/ObjectMixer/ObjectMixer/Source/ObjectMixer/Public/ObjectMixerEditorLog.h b/Engine/Plugins/Editor/ObjectMixer/ObjectMixer/Source/ObjectMixer/Public/ObjectMixerEditorLog.h
new file mode 100644
--- /dev/null
+++ b/Engine/Plugins/Editor/ObjectMixer/ObjectMixer/Source/ObjectMixer/Public/ObjectMixerEditorLog.h
@@ -0,0 +1,7 @@
+// Copyright Epic Games, Inc. All Rights Reserved.
+
+#pragma once
+
+#include "Logging/LogCategory.h"
+
+DECLARE_LOG_CATEGORY_CLASS(LogObjectMixerEditor, Log, All);
diff --git a/Engine/Plugins/Editor/ObjectMixer/ObjectMixer/Source/ObjectMixer/Public/ObjectMixerEditorModule.h b/Engine/Plugins/Editor/ObjectMixer/ObjectMixer/Source/ObjectMixer/Public/ObjectMixerEditorModule.h
new file mode 100644
--- /dev/null
+++ b/Engine/Plugins/Editor/ObjectMixer/ObjectMixer/Source/ObjectMixer/Public/ObjectMixerEditorModule.h
@@ -0,0 +1,90 @@
+// Copyright Epic Games, Inc. All Rights Reserved.
+
+#pragma once
+
+#include "Modules/ModuleInterface.h"
+#include "ObjectFilter/ObjectMixerEditorObjectFilter.h"
+#include "Widgets/Docking/SDockTab.h"
+
+class FObjectMixerEditorMainPanel;
+
+class OBJECTMIXEREDITOR_API FObjectMixerEditorModule : public IModuleInterface
+{
+public:
+
+	//~ Begin IModuleInterface Interface
+	virtual void StartupModule() override;
+	virtual void ShutdownModule() override;
+	//~ End IModuleInterface
+
+	virtual void Initialize();
+	virtual void Teardown();
+
+	static FObjectMixerEditorModule& Get();
+
+	static void OpenProjectSettings();
+
+	virtual FName GetModuleName();
+	
+	virtual TSharedPtr<SWidget> MakeObjectMixerDialog() const;
+
+	/**
+	 * Regenerate the list items and refresh the list. Call when adding or removing variables.
+	 */
+	virtual void RequestRebuildList() const;
+	
+	/**
+	 * Refresh filters and sorting.
+	 * Useful for when the list state has gone stale but the variable count has not changed.
+	 */
+	virtual void RefreshList() const;
+
+	void RegisterMenuGroup();
+	void UnregisterMenuGroup();
+	virtual void SetupMenuItemVariables();
+	virtual void RegisterTabSpawner();
+	virtual FName GetTabSpawnerId();
+	
+	/**
+	 * Add a tab spawner to the Object Mixer menu group.
+	 * @return If adding the item to the menu was successful
+	 */
+	bool RegisterItemInMenuGroup(FWorkspaceItem& InItem);
+	
+	virtual void UnregisterTabSpawner();
+	virtual void RegisterProjectSettings() const;
+	virtual void UnregisterProjectSettings() const;
+
+	virtual TSharedRef<SDockTab> SpawnTab(const FSpawnTabArgs& Args);
+
+	TSharedPtr<FWorkspaceItem> GetWorkspaceGroup();
+
+	const static FName BaseObjectMixerModuleName;
+
+protected:
+
+	virtual TSharedRef<SDockTab> SpawnMainPanelTab();
+
+	virtual void BindDelegates();
+	
+	/** Lives for as long as the module is loaded. */
+	TSharedPtr<FObjectMixerEditorMainPanel> MainPanel;
+
+	/** The text that appears on the spawned nomad tab */
+	FText TabLabel;
+
+	/** Menu Item variables */
+	FText MenuItemName;
+	FSlateIcon MenuItemIcon;
+	FText MenuItemTooltip;
+	ETabSpawnerMenuType::Type TabSpawnerType = ETabSpawnerMenuType::Enabled;
+
+	// If set, this is the filter class used to initialize the MainPanel
+	TSubclassOf<UObjectMixerObjectFilter> DefaultFilterClass;
+	
+	TSet<FDelegateHandle> DelegateHandles;
+
+private:
+
+	TSharedPtr<FWorkspaceItem> WorkspaceGroup;
+};
diff --git a/Engine/Plugins/Editor/ObjectMixer/ObjectMixer/Source/ObjectMixer/Public/ObjectMixerEditorProjectSettings.h b/Engine/Plugins/Editor/ObjectMixer/ObjectMixer/Source/ObjectMixer/Public/ObjectMixerEditorProjectSettings.h
new file mode 100644
--- /dev/null
+++ b/Engine/Plugins/Editor/ObjectMixer/ObjectMixer/Source/ObjectMixer/Public/ObjectMixerEditorProjectSettings.h
@@ -0,0 +1,29 @@
+// Copyright Epic Games, Inc. All Rights Reserved.
+
+#pragma once
+
+#include "ObjectMixerEditorProjectSettings.generated.h"
+
+UCLASS(config = ObjectMixer, defaultconfig)
+class OBJECTMIXEREDITOR_API UObjectMixerEditorProjectSettings : public UObject
+{
+	GENERATED_BODY()
+public:
+	
+	UObjectMixerEditorProjectSettings(const FObjectInitializer& ObjectInitializer)
+	{}
+	
+	/**
+	 * If enabled, clicking an item in the mixer list will also select the item in the Scene Outliner.
+	 * Alt + Click to select items in mixer without selecting the item in the Scene outliner.
+	 * If disabled, selections will not sync unless Alt is held. Effectively, this is the opposite behavior.
+	 */
+	UPROPERTY(Config, EditAnywhere, Category="Object Mixer")
+	bool bSyncSelection = true;
+
+	/**
+	 * If false, a new object will be created every time the filter object is accessed.
+	 */
+	UPROPERTY(Config, EditAnywhere, Category="Object Mixer")
+	bool bExpandTreeViewItemsByDefault = true;
+};
\ No newline at end of file
diff --git a/Engine/Plugins/Editor/ObjectMixer/ObjectMixer/Source/ObjectMixer/Public/ObjectMixerEditorSerializedData.h b/Engine/Plugins/Editor/ObjectMixer/ObjectMixer/Source/ObjectMixer/Public/ObjectMixerEditorSerializedData.h
new file mode 100644
--- /dev/null
+++ b/Engine/Plugins/Editor/ObjectMixer/ObjectMixer/Source/ObjectMixer/Public/ObjectMixerEditorSerializedData.h
@@ -0,0 +1,81 @@
+// Copyright Epic Games, Inc. All Rights Reserved.
+
+#pragma once
+
+#include "ObjectMixerEditorSerializedData.generated.h"
+
+USTRUCT()
+struct FObjectMixerCollectionObjectSet
+{
+	GENERATED_BODY()
+	
+	UPROPERTY()
+	FName CollectionName = NAME_None;
+
+	UPROPERTY()
+	TSet<FSoftObjectPath> CollectionObjects = {};
+
+	bool operator==(const FObjectMixerCollectionObjectSet& Other) const
+	{
+		return CollectionName.IsEqual(Other.CollectionName);
+	}
+
+	friend uint32 GetTypeHash (const FObjectMixerCollectionObjectSet& Other)
+	{
+		return GetTypeHash(Other.CollectionName);
+	}
+};
+
+USTRUCT()
+struct FObjectMixerSerializationData
+{
+	GENERATED_BODY()
+	
+	UPROPERTY()
+	FName FilterClassName = NAME_None;
+
+	UPROPERTY()
+	TArray<FObjectMixerCollectionObjectSet> SerializedCollection = {};
+
+	bool operator==(const FObjectMixerSerializationData& Other) const
+	{
+		return FilterClassName.IsEqual(Other.FilterClassName);
+	}
+
+	friend uint32 GetTypeHash (const FObjectMixerSerializationData& Other)
+	{
+		return GetTypeHash(Other.FilterClassName);
+	}
+};
+
+UCLASS(config = ObjectMixerSerializedData)
+class OBJECTMIXEREDITOR_API UObjectMixerEditorSerializedData : public UObject
+{
+	GENERATED_BODY()
+public:
+	
+	UObjectMixerEditorSerializedData(const FObjectInitializer& ObjectInitializer)
+	{
+		// Makes undo/redo possible for this object
+		SetFlags(RF_Transactional);
+	}
+	
+	UPROPERTY(Config)
+	TSet<FObjectMixerSerializationData> SerializedData;
+
+	FObjectMixerSerializationData* FindSerializationDataByFilterClassName(const FName& FilterClassName);
+
+	void AddObjectsToCollection(const FName& FilterClassName, const FName& CollectionName, const TSet<FSoftObjectPath>& ObjectsToAdd);
+
+	void RemoveObjectsFromCollection(const FName& FilterClassName, const FName& CollectionName, const TSet<FSoftObjectPath>& ObjectsToRemove);
+
+	void RemoveCollection(const FName& FilterClassName, const FName& CollectionName);
+
+	void ReorderCollection(const FName& FilterClassName, const FName& CollectionToMoveName, const FName& CollectionInsertBeforeName);
+
+	bool IsObjectInCollection(const FName& FilterClassName, const FName& CollectionName, const FSoftObjectPath& InObject);
+
+	TSet<FName> GetCollectionsForObject(const FName& FilterClassName, const FSoftObjectPath& InObject);
+
+	TArray<FName> GetAllCollectionNames(const FName& FilterClassName);
+};
diff --git a/Engine/Plugins/Editor/ObjectMixer/ObjectMixer/Source/ObjectMixer/Public/ObjectMixerEditorStyle.h b/Engine/Plugins/Editor/ObjectMixer/ObjectMixer/Source/ObjectMixer/Public/ObjectMixerEditorStyle.h
new file mode 100644
--- /dev/null
+++ b/Engine/Plugins/Editor/ObjectMixer/ObjectMixer/Source/ObjectMixer/Public/ObjectMixerEditorStyle.h
@@ -0,0 +1,36 @@
+// Copyright Epic Games, Inc. All Rights Reserved.
+
+#pragma once
+
+#include "Styling/SlateStyle.h"
+
+class OBJECTMIXEREDITOR_API FObjectMixerEditorStyle : FSlateStyleSet
+{
+public:
+
+	static void Initialize();
+
+	static void Shutdown();
+
+	static void ReloadTextures();
+
+	static const ISlateStyle& Get();
+
+	virtual const FName& GetStyleSetName() const override;
+
+	virtual const FSlateBrush* GetBrush(const FName PropertyName, const ANSICHAR* Specifier = nullptr, const ISlateStyle* RequestingStyle = nullptr) const override;
+
+	template< typename WidgetStyleType >
+	static const WidgetStyleType& GetWidgetStyle(FName PropertyName, const ANSICHAR* Specifier = nullptr)
+	{
+		return StyleInstance->GetWidgetStyle<WidgetStyleType>(PropertyName, Specifier);
+	}
+
+private:
+
+	static FString GetExternalPluginContent(const FString& PluginName, const FString& RelativePath, const ANSICHAR* Extension);
+	
+	static TSharedRef<FSlateStyleSet> Create();
+
+	static TSharedPtr<FSlateStyleSet> StyleInstance;
+};
diff --git a/Engine/Plugins/Editor/ObjectMixer/ObjectMixer/Source/ObjectMixer/Public/ObjectMixerFilterFactory.h b/Engine/Plugins/Editor/ObjectMixer/ObjectMixer/Source/ObjectMixer/Public/ObjectMixerFilterFactory.h
new file mode 100644
--- /dev/null
+++ b/Engine/Plugins/Editor/ObjectMixer/ObjectMixer/Source/ObjectMixer/Public/ObjectMixerFilterFactory.h
@@ -0,0 +1,27 @@
+// Copyright Epic Games, Inc. All Rights Reserved.
+
+#pragma once
+
+#include "Factories/Factory.h"
+
+#include "ObjectMixerFilterFactory.generated.h"
+
+UCLASS(hidecategories = Object)
+class OBJECTMIXEREDITOR_API UObjectMixerBlueprintFilterFactory : public UFactory
+{
+	GENERATED_BODY()
+public:
+	UObjectMixerBlueprintFilterFactory();
+	
+	// The parent class of the created blueprint
+    UPROPERTY(VisibleAnywhere, Category="ObjectMixerBlueprintFilterFactory", meta=(AllowAbstract = "", BlueprintBaseOnly = ""))
+    TSubclassOf<UObject> ParentClass;
+	
+	//~ Begin UFactory Interface
+	virtual FText GetDisplayName() const override;
+	virtual UObject* FactoryCreateNew(UClass* Class, UObject* InParent, FName Name, EObjectFlags Flags, UObject* Context, FFeedbackContext* Warn) override;
+	virtual FString GetDefaultNewAssetName() const override;
+	virtual uint32 GetMenuCategories() const override;
+	//~ Begin UFactory Interface
+	
+};
diff --git a/Engine/Plugins/Editor/ObjectMixer/ObjectMixer/Source/ObjectMixer/Public/Views/List/ObjectMixerEditorList.h b/Engine/Plugins/Editor/ObjectMixer/ObjectMixer/Source/ObjectMixer/Public/Views/List/ObjectMixerEditorList.h
new file mode 100644
--- /dev/null
+++ b/Engine/Plugins/Editor/ObjectMixer/ObjectMixer/Source/ObjectMixer/Public/Views/List/ObjectMixerEditorList.h
@@ -0,0 +1,56 @@
+// Copyright Epic Games, Inc. All Rights Reserved.
+
+#pragma once
+
+#include "Views/MainPanel/ObjectMixerEditorMainPanel.h"
+
+#include "Templates/SharedPointer.h"
+#include "Widgets/SWidget.h"
+
+class SObjectMixerEditorList;
+
+class OBJECTMIXEREDITOR_API FObjectMixerEditorList : public TSharedFromThis<FObjectMixerEditorList>
+{
+public:
+
+	FObjectMixerEditorList(TSharedRef<FObjectMixerEditorMainPanel, ESPMode::ThreadSafe> InMainPanel)
+	: MainPanelModelPtr(InMainPanel)
+	{}
+
+	virtual ~FObjectMixerEditorList();
+
+	void FlushWidget();
+	
+	TSharedRef<SWidget> GetOrCreateWidget();
+
+	void ClearList() const;
+
+	/**
+	 * Regenerate the list items and refresh the list. Call when adding or removing variables.
+	 */
+	void RequestRebuildList() const;
+
+	/**
+	 * Refresh filters and sorting.
+	 * Useful for when the list state has gone stale but the item count has not changed.
+	 */
+	void RefreshList() const;
+
+	void ExecuteListViewSearchOnAllRows(const FString& SearchString, const bool bShouldRefreshAfterward = true);
+
+	void EvaluateIfRowsPassFilters(const bool bShouldRefreshAfterward = true) const;
+
+	TWeakPtr<FObjectMixerEditorMainPanel> GetMainPanelModel();
+
+	TWeakPtr<FObjectMixerEditorListRow> GetSoloRow();
+
+	void SetSoloRow(TSharedRef<FObjectMixerEditorListRow> InRow);
+
+	void ClearSoloRow();
+
+private:
+
+	TWeakPtr<FObjectMixerEditorMainPanel> MainPanelModelPtr;
+
+	TSharedPtr<SObjectMixerEditorList> ListWidget;
+};
diff --git a/Engine/Plugins/Editor/ObjectMixer/ObjectMixer/Source/ObjectMixer/Public/Views/List/ObjectMixerEditorListFilters/IObjectMixerEditorListFilter.h b/Engine/Plugins/Editor/ObjectMixer/ObjectMixer/Source/ObjectMixer/Public/Views/List/ObjectMixerEditorListFilters/IObjectMixerEditorListFilter.h
new file mode 100644
--- /dev/null
+++ b/Engine/Plugins/Editor/ObjectMixer/ObjectMixer/Source/ObjectMixer/Public/Views/List/ObjectMixerEditorListFilters/IObjectMixerEditorListFilter.h
@@ -0,0 +1,66 @@
+// Copyright Epic Games, Inc. All Rights Reserved.
+
+#pragma once
+
+#include "../ObjectMixerEditorListRow.h"
+
+class OBJECTMIXEREDITOR_API IObjectMixerEditorListFilter
+{
+public:
+
+	enum class EObjectMixerEditorListFilterMatchType
+	{
+		// Filters of this type are placed into a group in which if any of the group's filters pass, the row is considered passing
+		MatchAny,
+		// Filters of this type are placed into a group in which all of the group's filters must pass for the row to be considered passing
+		MatchAll
+	};
+	
+	virtual ~IObjectMixerEditorListFilter() = default;
+
+	virtual FString GetFilterName() const = 0;
+
+	/** Return true if the filter should be able to be toggled on/off in the Show Options. */
+	virtual bool IsToggleable() const = 0;
+
+	/** Returns localized text to display for when this filter' active state is defined by a toggle button. */
+	virtual FText GetFilterButtonLabel() const = 0;
+	/** Returns localized text to display when the user mouses over a toggle button that defines whether this filter is active. */
+	virtual FText GetFilterButtonToolTip() const = 0;
+	
+	void SetFilterActive(const bool bNewEnabled)
+	{
+		bIsFilterActive = bNewEnabled;
+	}
+
+	void ToggleFilterActive()
+	{
+		bIsFilterActive = !bIsFilterActive;
+	}
+
+	bool GetIsFilterActive() const
+	{
+		return !IsToggleable() || bIsFilterActive;
+	}
+
+	void SetFilterMatchType(const EObjectMixerEditorListFilterMatchType InType)
+	{
+		FilterMatchType = InType;
+	}
+
+	EObjectMixerEditorListFilterMatchType GetFilterMatchType() const
+	{
+		return FilterMatchType;
+	}
+
+	virtual bool DoesItemPassFilter(const FObjectMixerEditorListRowPtr& InItem) const
+	{
+		return false;
+	}
+
+private:
+
+	bool bIsFilterActive = true;
+	EObjectMixerEditorListFilterMatchType FilterMatchType = EObjectMixerEditorListFilterMatchType::MatchAny;
+			
+};
diff --git a/Engine/Plugins/Editor/ObjectMixer/ObjectMixer/Source/ObjectMixer/Public/Views/List/ObjectMixerEditorListFilters/ObjectMixerEditorListFilter_Collection.h b/Engine/Plugins/Editor/ObjectMixer/ObjectMixer/Source/ObjectMixer/Public/Views/List/ObjectMixerEditorListFilters/ObjectMixerEditorListFilter_Collection.h
new file mode 100644
--- /dev/null
+++ b/Engine/Plugins/Editor/ObjectMixer/ObjectMixer/Source/ObjectMixer/Public/Views/List/ObjectMixerEditorListFilters/ObjectMixerEditorListFilter_Collection.h
@@ -0,0 +1,54 @@
+// Copyright Epic Games, Inc. All Rights Reserved.
+
+#pragma once
+
+#include "IObjectMixerEditorListFilter.h"
+
+#define LOCTEXT_NAMESPACE "ObjectMixerEditor"
+
+class OBJECTMIXEREDITOR_API FObjectMixerEditorListFilter_Collection : public IObjectMixerEditorListFilter
+{
+public:
+
+	FObjectMixerEditorListFilter_Collection()
+	{
+		SetFilterMatchType(EObjectMixerEditorListFilterMatchType::MatchAll);
+	}
+
+	virtual FString GetFilterName() const override
+	{
+		return "ObjectMixerCollectionListFilter";
+	}
+
+	virtual bool IsToggleable() const override
+	{
+		return false;
+	}
+
+	virtual FText GetFilterButtonLabel() const override
+	{
+		return LOCTEXT("ShowSourceTextFilterFormat", "Show Collections");
+	}
+
+	virtual FText GetFilterButtonToolTip() const override
+	{
+		return LOCTEXT("ShowSourceTextFilterTooltipFormat", "Show rows that are assigned to the selected collections");
+	}
+
+	virtual bool DoesItemPassFilter(const FObjectMixerEditorListRowPtr& InItem) const override
+	{
+		if (InItem.IsValid())
+		{
+			if (InItem->GetRowType() == FObjectMixerEditorListRow::ContainerObject ||
+				InItem->GetRowType() == FObjectMixerEditorListRow::MatchingObject)
+			{
+				return InItem->IsObjectRefInSelectedCollections();
+			}
+		}
+
+		return false;
+	}
+
+};
+
+#undef LOCTEXT_NAMESPACE
diff --git a/Engine/Plugins/Editor/ObjectMixer/ObjectMixer/Source/ObjectMixer/Public/Views/List/ObjectMixerEditorListRow.h b/Engine/Plugins/Editor/ObjectMixer/ObjectMixer/Source/ObjectMixer/Public/Views/List/ObjectMixerEditorListRow.h
new file mode 100644
--- /dev/null
+++ b/Engine/Plugins/Editor/ObjectMixer/ObjectMixer/Source/ObjectMixer/Public/Views/List/ObjectMixerEditorListRow.h
@@ -0,0 +1,172 @@
+// Copyright Epic Games, Inc. All Rights Reserved.
+
+#pragma once
+
+#include "Containers/Array.h"
+#include "Containers/UnrealString.h"
+#include "DragAndDrop/DecoratedDragDropOp.h"
+#include "Layout/Visibility.h"
+#include "ObjectFilter/ObjectMixerEditorObjectFilter.h"
+#include "Templates/SharedPointer.h"
+#include "UObject/Object.h"
+#include "UObject/WeakObjectPtr.h"
+
+class SObjectMixerEditorList;
+class UObjectMixerObjectFilter;
+
+struct FSlateBrush;
+
+struct FObjectMixerEditorListRow;
+typedef TSharedPtr<FObjectMixerEditorListRow> FObjectMixerEditorListRowPtr;
+
+class FObjectMixerListRowDragDropOp : public FDecoratedDragDropOp
+{
+public:
+	DRAG_DROP_OPERATOR_TYPE(FObjectMixerListRowDragDropOp, FDecoratedDragDropOp)
+
+	/** The item being dragged and dropped */
+	TArray<FObjectMixerEditorListRowPtr> DraggedItems;
+
+	/** Constructs a new drag/drop operation */
+	static TSharedRef<FObjectMixerListRowDragDropOp> New(const TArray<FObjectMixerEditorListRowPtr>& InItems);
+};
+
+struct OBJECTMIXEREDITOR_API FObjectMixerEditorListRow final : TSharedFromThis<FObjectMixerEditorListRow>
+{
+	enum EObjectMixerEditorListRowType : uint8
+	{
+		None = 0,
+		Folder, // Usually an Outliner folder
+		ContainerObject, // Usually an actor that contains a matching subobject or is the attach parent of a matching actor
+		MatchingObject // The object that has the properties we wish to affect
+	};
+
+	~FObjectMixerEditorListRow();
+
+	void FlushReferences();
+	
+	FObjectMixerEditorListRow(
+		const TWeakObjectPtr<UObject> InObject, const EObjectMixerEditorListRowType InRowType, 
+		const TSharedRef<SObjectMixerEditorList>& InListView, const FText& InDisplayNameOverride = FText::GetEmpty())
+	: ObjectRef(InObject)
+	, RowType(InRowType)
+	, ListViewPtr(InListView)
+	, DisplayNameOverride(InDisplayNameOverride)
+	{}
+
+	[[nodiscard]] UObject* GetObject() const
+	{
+		if (ObjectRef.IsValid())
+		{
+			return ObjectRef.Get();
+		}
+
+		return nullptr;
+	}
+
+	UObjectMixerObjectFilter* GetObjectFilter() const;
+
+	bool IsObjectRefInSelectedCollections() const;
+
+	[[nodiscard]] EObjectMixerEditorListRowType GetRowType() const;
+
+	[[nodiscard]] int32 GetSortOrder() const;
+	void SetSortOrder(const int32 InNewOrder);
+
+	TWeakPtr<FObjectMixerEditorListRow> GetDirectParentRow() const;
+	void SetDirectParentRow(const TWeakPtr<FObjectMixerEditorListRow>& InDirectParentRow);
+	
+	/** bHasGeneratedChildren must be true to get actual children. */
+	[[nodiscard]] const TArray<FObjectMixerEditorListRowPtr>& GetChildRows() const;
+	/** bHasGeneratedChildren must be true to get an accurate value. */
+	[[nodiscard]] int32 GetChildCount() const;
+	void SetChildRows(const TArray<FObjectMixerEditorListRowPtr>& InChildRows);
+	void AddToChildRows(const FObjectMixerEditorListRowPtr& InRow);
+	void InsertChildRowAtIndex(const FObjectMixerEditorListRowPtr& InRow, const int32 AtIndex = 0);
+
+	[[nodiscard]] bool GetIsTreeViewItemExpanded();
+	void SetIsTreeViewItemExpanded(const bool bNewExpanded);
+
+	[[nodiscard]] bool GetShouldExpandAllChildren() const;
+	void SetShouldExpandAllChildren(const bool bNewShouldExpandAllChildren);
+
+	/**
+	 *Individual members of InTokens will be considered "AnyOf" or "OR" searches. If SearchTerms contains any individual member it will match.
+	 *Members will be tested for a space character (" "). If a space is found, a subsearch will be run.
+	 *This subsearch will be an "AllOf" or "AND" type search in which all strings, separated by a space, must be found in the search terms.
+	 */
+	bool MatchSearchTokensToSearchTerms(const TArray<FString> InTokens, ESearchCase::Type InSearchCase = ESearchCase::IgnoreCase);
+
+	/** This overload creates tokens from a string first, then calls ExecuteSearchOnChildNodes(const TArray<FString>& Tokens). */
+	void ExecuteSearchOnChildNodes(const FString& SearchString) const;
+	void ExecuteSearchOnChildNodes(const TArray<FString>& Tokens) const;
+
+	[[nodiscard]] bool GetDoesRowPassFilters() const;
+	void SetDoesRowPassFilters(const bool bPass);
+
+	[[nodiscard]] bool GetIsSelected();
+
+	[[nodiscard]] bool ShouldRowWidgetBeVisible() const;
+	[[nodiscard]] EVisibility GetDesiredRowWidgetVisibility() const;
+
+	[[nodiscard]] bool HasVisibleChildRowWidgets() const;
+
+	[[nodiscard]] FText GetDisplayName();
+
+	[[nodiscard]] const FText& GetDisplayNameOverride()
+	{
+		return DisplayNameOverride;
+	}
+
+	void SetDisplayNameOverride(const FText& InDisplayNameOverride)
+	{
+		DisplayNameOverride = InDisplayNameOverride;
+	}
+
+	[[nodiscard]] TWeakPtr<SObjectMixerEditorList> GetListViewPtr() const
+	{
+		return ListViewPtr;
+	}
+
+	/**
+	 * Determines the style of the tree (flat list or hierarchy)
+	 */
+	EObjectMixerTreeViewMode GetTreeViewMode();
+
+	[[nodiscard]] TArray<FObjectMixerEditorListRowPtr> GetSelectedTreeViewItems() const;
+
+	const FSlateBrush* GetObjectIconBrush();
+
+	bool GetObjectVisibility();
+	void SetObjectVisibility(const bool bNewIsVisible, const bool bIsRecursive = false);
+
+	bool IsThisRowSolo() const;
+	void SetThisAsSoloRow();
+
+	void ClearSoloRow();
+
+private:
+
+	FObjectMixerEditorListRowPtr GetAsShared();
+
+	TWeakObjectPtr<UObject> ObjectRef;
+	EObjectMixerEditorListRowType RowType = MatchingObject;
+	TArray<FObjectMixerEditorListRowPtr> ChildRows;
+	
+	TWeakPtr<SObjectMixerEditorList> ListViewPtr;
+
+	FText DisplayNameOverride;
+
+	int32 SortOrder = -1;
+
+	FString CachedSearchTerms;
+
+	bool bDoesRowMatchSearchTerms = true;
+	bool bDoesRowPassFilters = true;
+	
+	bool bIsSelected = false;
+	TWeakPtr<FObjectMixerEditorListRow> DirectParentRow;
+
+	// Used to expand all children on shift+click.
+	bool bShouldExpandAllChildren = false;
+};
diff --git a/Engine/Plugins/Editor/ObjectMixer/ObjectMixer/Source/ObjectMixer/Public/Views/List/SObjectMixerEditorList.h b/Engine/Plugins/Editor/ObjectMixer/ObjectMixer/Source/ObjectMixer/Public/Views/List/SObjectMixerEditorList.h
new file mode 100644
--- /dev/null
+++ b/Engine/Plugins/Editor/ObjectMixer/ObjectMixer/Source/ObjectMixer/Public/Views/List/SObjectMixerEditorList.h
@@ -0,0 +1,256 @@
+// Copyright Epic Games, Inc. All Rights Reserved.
+
+#pragma once
+
+#include "ObjectMixerEditorList.h"
+#include "ObjectMixerEditorListRow.h"
+#include "ObjectMixerEditorListFilters/IObjectMixerEditorListFilter.h"
+
+#include "Widgets/Input/SComboButton.h"
+#include "Widgets/Layout/SGridPanel.h"
+#include "Widgets/SCompoundWidget.h"
+#include "Widgets/Views/SHeaderRow.h"
+#include "Widgets/Views/STreeView.h"
+
+#include "SObjectMixerEditorList.generated.h"
+
+class SWrapBox;
+class FObjectMixerEditorList;
+class SBox;
+class SComboButton;
+class SSearchBox;
+class SHeaderRow;
+
+UENUM()
+enum class EListViewColumnType
+{
+	BuiltIn,
+	PropertyGenerated
+};
+
+USTRUCT()
+struct FListViewColumnInfo
+{
+	GENERATED_BODY()
+	
+	TObjectPtr<FProperty> PropertyRef = nullptr;
+	
+	UPROPERTY()
+	FName PropertyName = NAME_None;
+
+	UPROPERTY()
+	FText PropertyDisplayText = FText::GetEmpty();
+
+	UPROPERTY()
+	EListViewColumnType PropertyType = EListViewColumnType::PropertyGenerated;
+
+	UPROPERTY()
+	FName PropertyCategoryName = NAME_None;
+
+	UPROPERTY()
+	bool bIsDesiredForDisplay = false;
+
+	UPROPERTY()
+	bool bCanBeSorted = false;
+
+	UPROPERTY()
+	bool bUseFixedWidth = false;
+
+	UPROPERTY()
+	float FixedWidth = 25.0f;
+
+	UPROPERTY()
+	float FillWidth = 1.0f;
+};
+
+class OBJECTMIXEREDITOR_API SObjectMixerEditorList final : public SCompoundWidget
+{
+
+public:
+
+	SLATE_BEGIN_ARGS(SObjectMixerEditorList)
+	{}
+
+	SLATE_END_ARGS()
+
+	void Construct(const FArguments& InArgs, TSharedRef<FObjectMixerEditorList> ListModel);
+
+	virtual ~SObjectMixerEditorList() override;
+
+	TWeakPtr<FObjectMixerEditorList> GetListModelPtr()
+	{
+		return ListModelPtr;
+	}
+
+	void ClearList()
+	{
+		FlushMemory(false);
+	}
+
+	/**
+	 * Refresh filters and sorting.
+	 * Useful for when the list state has gone stale but the variable count has not changed.
+	 */
+	void RefreshList();
+
+	/**
+	 * Regenerate the list items and refresh the list. Call when adding or removing items.
+	 */
+	void RequestRebuildList(const FString& InItemToScrollTo = "");
+
+	[[nodiscard]] TArray<FObjectMixerEditorListRowPtr> GetSelectedTreeViewItems() const;
+	int32 GetSelectedTreeViewItemCount() const;
+
+	void SetSelectedTreeViewItemActorsEditorVisible(const bool bNewIsVisible, const bool bIsRecursive = false);
+
+	bool IsTreeViewItemSelected(TSharedRef<FObjectMixerEditorListRow> Item);
+
+	[[nodiscard]] TArray<FObjectMixerEditorListRowPtr> GetTreeViewItems() const;
+	void SetTreeViewItems(const TArray<FObjectMixerEditorListRowPtr>& InItems);
+
+	[[nodiscard]] int32 GetTreeViewItemCount() const
+	{
+		return TreeViewRootObjects.Num();
+	}
+	
+	TWeakPtr<FObjectMixerEditorListRow> GetSoloRow();
+
+	void SetSoloRow(TSharedRef<FObjectMixerEditorListRow> InRow);
+
+	void ClearSoloRow();
+
+	FString GetSearchStringFromSearchInputField() const;
+	void ExecuteListViewSearchOnAllRows(const FString& SearchString, const bool bShouldRefreshAfterward = true);
+
+	bool DoesTreeViewHaveVisibleChildren() const;
+
+	bool IsTreeViewItemExpanded(const TSharedPtr<FObjectMixerEditorListRow>& Row) const;
+	void SetTreeViewItemExpanded(const TSharedPtr<FObjectMixerEditorListRow>& RowToExpand, const bool bNewExpansion) const;
+
+	/**
+	 * Determines the style of the tree (flat list or hierarchy)
+	 */
+	EObjectMixerTreeViewMode GetTreeViewMode();
+	
+	const TArray<TSharedRef<IObjectMixerEditorListFilter>>& GetShowFilters();
+	void EvaluateIfRowsPassFilters(const bool bShouldRefreshAfterward = true);
+
+	// Sorting
+
+	const FName& GetActiveSortingColumnName() const
+	{
+		return ActiveSortingColumnName;
+	}
+	EColumnSortMode::Type GetSortModeForColumn(FName InColumnName) const;
+	void OnSortColumnCalled(EColumnSortPriority::Type Priority, const FName& ColumnName, EColumnSortMode::Type SortMode);
+	EColumnSortMode::Type CycleSortMode(const FName& InColumnName);
+	void ExecuteSort(
+		const FName& InColumnName, const EColumnSortMode::Type InColumnSortMode, const bool bShouldRefreshAfterward = true);
+	void ClearSorting()
+	{
+		ActiveSortingColumnName = NAME_None;
+		ActiveSortingType = EColumnSortMode::None;
+	}
+
+	// Columns
+
+	FListViewColumnInfo* GetColumnInfoByPropertyName(const FName& InPropertyName);
+
+	static const FName ItemNameColumnName;
+	static const FName EditorVisibilityColumnName;
+	static const FName EditorVisibilitySoloColumnName;
+
+	inline static TFunction<bool(const FObjectMixerEditorListRowPtr&, const FObjectMixerEditorListRowPtr&)> SortByTypeThenName =
+		[](const FObjectMixerEditorListRowPtr& A, const FObjectMixerEditorListRowPtr& B)
+	{
+		return (A->GetRowType() < B->GetRowType() ||
+			(A->GetRowType() == B->GetRowType() && A->GetDisplayName().ToString() < B->GetDisplayName().ToString())
+		);
+	};
+
+protected:
+
+	virtual void Tick(const FGeometry& AllottedGeometry, const double InCurrentTime, const float InDeltaTime) override;
+
+	void OnActorSpawnedOrDestroyed(AActor* Object)
+	{
+		RequestRebuildList();
+	}
+
+	TWeakPtr<FObjectMixerEditorList> ListModelPtr;
+	
+	TSharedPtr<SHeaderRow> HeaderRow;
+	TSharedRef<SWidget> GenerateHeaderRowContextMenu() const;
+
+	bool bShouldRebuild = false;
+	
+	/**
+	 * Regenerate the list items and refresh the list. Call when adding or removing items.
+	 */
+	void RebuildList();
+
+	/**
+	 * Only adds properties that pass a series of tests, including having only one unique entry in the column list array.
+	 * @param bForceIncludeProperty If true, only Skiplist and Uniqueness tests will be checked, bypassing class, blueprint editability and other requirements.
+	 * @param PropertySkipList These property names will be skipped when they are encountered in the iteration.
+	 */
+	bool AddUniquePropertyColumnsToHeaderRow(
+		FProperty* Property,
+		const bool bForceIncludeProperty = false,
+		const TSet<FName>& PropertySkipList = {}
+	);
+	void AddBuiltinColumnsToHeaderRow();
+	TSharedPtr<SHeaderRow> GenerateHeaderRow();
+	ECheckBoxState HeaderCheckBoxState = ECheckBoxState::Checked;
+	
+	void FlushMemory(const bool bShouldKeepMemoryAllocated);
+	
+	void SetAllGroupsCollapsed();
+
+	//  Tree View Implementation
+
+	/** Saves tree item expanded states to be recalled after the tree view is regenerated. */
+	void CacheTreeState();
+	void RestoreTreeState(const bool bFlushCache = true);
+	void BuildPerformanceCacheAndGenerateHeaderIfNeeded();
+	void GenerateTreeView();
+	void FindVisibleTreeViewObjects();
+	void FindVisibleObjectsAndRequestTreeRefresh();
+
+	void SelectedTreeItemsToSelectedInLevelEditor() const;
+	
+	void OnGetRowChildren(FObjectMixerEditorListRowPtr Row, TArray<FObjectMixerEditorListRowPtr>& OutChildren) const;
+	void OnRowChildExpansionChange(FObjectMixerEditorListRowPtr Row, const bool bIsExpanded, const bool bIsRecursive = false) const;
+
+	void SetChildExpansionRecursively(const FObjectMixerEditorListRowPtr& InRow, const bool bNewIsExpanded) const;
+
+	TSharedPtr<STreeView<FObjectMixerEditorListRowPtr>> TreeViewPtr;
+
+	TMap<FString, bool> TreeItemExpansionStateCache;
+
+	/** All Tree view objects */
+	TArray<FObjectMixerEditorListRowPtr> TreeViewRootObjects;
+	/** Visible Tree view objects, after filters */
+	TArray<FObjectMixerEditorListRowPtr> VisibleTreeViewObjects;
+
+	TArray<FListViewColumnInfo> ListViewColumns;
+
+	// Sorting
+
+	FName ActiveSortingColumnName = NAME_None;
+	EColumnSortMode::Type ActiveSortingType = EColumnSortMode::None;
+
+	TFunctionRef<bool(const FObjectMixerEditorListRowPtr&, const FObjectMixerEditorListRowPtr&)> SortByOrderAscending =
+		[](const FObjectMixerEditorListRowPtr& A, const FObjectMixerEditorListRowPtr& B)
+		{
+			return A->GetSortOrder() < B->GetSortOrder();
+		};
+
+	// Performance cache
+	TSet<UClass*> ObjectClassesToFilterCache;
+	TSet<FName> ColumnsToShowByDefaultCache;
+	TSet<FName> ColumnsToExcludeCache;
+	TSet<FName> ForceAddedColumnsCache;
+	EObjectMixerInheritanceInclusionOptions PropertyInheritanceInclusionOptionsCache = EObjectMixerInheritanceInclusionOptions::None;
+	bool bShouldIncludeUnsupportedPropertiesCache = false;
+};
diff --git a/Engine/Plugins/Editor/ObjectMixer/ObjectMixer/Source/ObjectMixer/Public/Views/List/SObjectMixerEditorListRow.h b/Engine/Plugins/Editor/ObjectMixer/ObjectMixer/Source/ObjectMixer/Public/Views/List/SObjectMixerEditorListRow.h
new file mode 100644
--- /dev/null
+++ b/Engine/Plugins/Editor/ObjectMixer/ObjectMixer/Source/ObjectMixer/Public/Views/List/SObjectMixerEditorListRow.h
@@ -0,0 +1,67 @@
+// Copyright Epic Games, Inc. All Rights Reserved.
+
+#pragma once
+
+#include "ObjectMixerEditorListRow.h"
+
+#include "DragAndDrop/DecoratedDragDropOp.h"
+#include "Widgets/Input/SNumericEntryBox.h"
+#include "Widgets/IToolTip.h"
+#include "Widgets/Views/STableRow.h"
+
+class SObjectMixerEditorListValueInput;
+class SObjectMixerEditorListRowHoverWidgets;
+
+class OBJECTMIXEREDITOR_API SObjectMixerEditorListRow : public SMultiColumnTableRow<FObjectMixerEditorListRowPtr>
+{
+public:
+	
+	SLATE_BEGIN_ARGS(SObjectMixerEditorListRow)
+	{}
+
+	SLATE_END_ARGS()
+	
+	void Construct(const FArguments& InArgs, const TSharedRef<STableViewBase>& InOwnerTable, const TWeakPtr<FObjectMixerEditorListRow> InRow);
+	
+	virtual TSharedRef<SWidget> GenerateWidgetForColumn(const FName& InColumnName) override;
+
+	// Begin SWidget
+	virtual void OnMouseEnter(const FGeometry& MyGeometry, const FPointerEvent& MouseEvent) override;
+	virtual void OnMouseLeave(const FPointerEvent& MouseEvent) override;
+	// End SWidget
+	
+	virtual ~SObjectMixerEditorListRow() override;	
+
+	FReply HandleDragDetected(const FGeometry& MyGeometry, const FPointerEvent& MouseEvent);
+	void HandleDragLeave(const FDragDropEvent& DragDropEvent);
+	TOptional<EItemDropZone> HandleCanAcceptDrop(const FDragDropEvent& DragDropEvent, EItemDropZone DropZone, FObjectMixerEditorListRowPtr TargetItem);
+	FReply HandleAcceptDrop(const FDragDropEvent& DragDropEvent, EItemDropZone DropZone, FObjectMixerEditorListRowPtr TargetItem);
+	
+private:
+
+	bool IsVisible() const;
+
+	FSlateColor GetVisibilityIconForegroundColor() const;
+	FSlateColor GetSoloIconForegroundColor() const;
+	
+	/** Get the brush for this widget */
+	const FSlateBrush* GetVisibilityBrush() const;
+
+	static const FSlateBrush* GetBorderImage(const FObjectMixerEditorListRow::EObjectMixerEditorListRowType InRowType);
+
+	TSharedPtr<SWidget> GenerateCells(const FName& InColumnName, const TSharedPtr<FObjectMixerEditorListRow> PinnedItem);
+
+	void OnPropertyChanged(const FProperty* Property, void* ContainerWithChangedProperty);
+	
+	TWeakPtr<FObjectMixerEditorListRow> Item;
+	
+	const FSlateBrush* VisibleHoveredBrush = nullptr;
+	const FSlateBrush* VisibleNotHoveredBrush = nullptr;
+	const FSlateBrush* NotVisibleHoveredBrush = nullptr;
+	const FSlateBrush* NotVisibleNotHoveredBrush = nullptr;
+
+	/** The offset applied to text widgets so that the text aligns with the column header text */
+	float TextBlockLeftPadding = 3.0f;
+	
+	bool bIsHovered = false;
+};
diff --git a/Engine/Plugins/Editor/ObjectMixer/ObjectMixer/Source/ObjectMixer/Public/Views/MainPanel/ObjectMixerEditorMainPanel.h b/Engine/Plugins/Editor/ObjectMixer/ObjectMixer/Source/ObjectMixer/Public/Views/MainPanel/ObjectMixerEditorMainPanel.h
new file mode 100644
--- /dev/null
+++ b/Engine/Plugins/Editor/ObjectMixer/ObjectMixer/Source/ObjectMixer/Public/Views/MainPanel/ObjectMixerEditorMainPanel.h
@@ -0,0 +1,194 @@
+// Copyright Epic Games, Inc. All Rights Reserved.
+
+#pragma once
+
+#include "ObjectFilter/ObjectMixerEditorObjectFilter.h"
+#include "Views/List/ObjectMixerEditorListFilters/IObjectMixerEditorListFilter.h"
+
+#include "UObject/StrongObjectPtr.h"
+#include "Widgets/SWidget.h"
+
+class FObjectMixerEditorList;
+class SObjectMixerEditorMainPanel;
+
+DECLARE_MULTICAST_DELEGATE(FOnObjectMixerCollectionMapChanged)
+
+class OBJECTMIXEREDITOR_API FObjectMixerEditorMainPanel : public TSharedFromThis<FObjectMixerEditorMainPanel>
+{
+public:
+	FObjectMixerEditorMainPanel(const FName InModuleName)
+	: ModuleName(InModuleName)
+	{}
+
+	~FObjectMixerEditorMainPanel() = default;
+
+	void Init();
+
+	TSharedRef<SWidget> GetOrCreateWidget();
+
+	void RegenerateListModel();
+
+	/**
+	 * Regenerate the list items and refresh the list. Call when adding or removing items.
+	 */
+	void RequestRebuildList() const;
+
+	/**
+	 * Refresh filters and sorting.
+	 * Useful for when the list state has gone stale but the variable count has not changed.
+	 */
+	void RefreshList() const;
+
+	TWeakPtr<FObjectMixerEditorList> GetEditorListModel() const
+	{
+		return EditorListModel;
+	}
+
+	void RebuildCollectionSelector();
+
+	FString GetSearchStringFromSearchInputField() const;
+
+	void OnClassSelectionChanged(UClass* InNewClass);
+	TObjectPtr<UClass> GetClassSelection() const;
+	bool IsClassSelected(UClass* InNewClass) const;
+
+	UObjectMixerObjectFilter* GetObjectFilter();
+
+	void CacheObjectFilterObject();
+
+	/**
+	 * Get the style of the tree (flat list or hierarchy)
+	 */
+	EObjectMixerTreeViewMode GetTreeViewMode()
+	{
+		return TreeViewMode;
+	}
+	/**
+	 * Set the style of the tree (flat list or hierarchy)
+	 */
+	void SetTreeViewMode(EObjectMixerTreeViewMode InViewMode)
+	{
+		TreeViewMode = InViewMode;
+		RequestRebuildList();
+	}
+
+	/**
+	 * Returns result from Filter->GetObjectClassesToFilter.
+	 */
+	TSet<UClass*> GetObjectClassesToFilter()
+	{
+		if (const UObjectMixerObjectFilter* Filter = GetObjectFilter())
+		{
+			return Filter->GetObjectClassesToFilter();
+		}
+		
+		return {};
+	}
+
+	/**
+	 * Returns result from Filter->GetObjectClassesToPlace.
+	 */
+	TSet<TSubclassOf<AActor>> GetObjectClassesToPlace()
+	{
+		TSet<TSubclassOf<AActor>> ReturnValue;
+
+		if (const UObjectMixerObjectFilter* Filter = GetObjectFilter())
+		{
+			ReturnValue = Filter->GetObjectClassesToPlace();
+		}
+		
+		return ReturnValue;
+	}
+
+	const TArray<TSharedRef<IObjectMixerEditorListFilter>>& GetShowFilters() const;
+
+	/**
+	 * Get the row that has solo visibility. All other rows should be set to temporarily invisible in editor.
+	 */
+	TWeakPtr<FObjectMixerEditorListRow> GetSoloRow()
+	{
+		return SoloRow;
+	}
+
+	/**
+	 * Set the row that has solo visibility. This does not set temporary editor invisibility for other rows.
+	 */
+	void SetSoloRow(TSharedRef<FObjectMixerEditorListRow> InRow)
+	{
+		SoloRow = InRow;
+	}
+
+	/**
+	 * Clear the row that has solo visibility. This does not remove temporary editor invisibility for other rows.
+	 */
+	void ClearSoloRow()
+	{
+		SoloRow = nullptr;
+	}
+
+	TSubclassOf<UObjectMixerObjectFilter> GetObjectFilterClass() const
+	{
+		return ObjectFilterClass;
+	}
+
+	void SetObjectFilterClass(UClass* InObjectFilterClass)
+	{
+		if (ensureAlwaysMsgf(InObjectFilterClass->IsChildOf(UObjectMixerObjectFilter::StaticClass()), TEXT("%hs: Class '%s' is not a child of UObjectMixerObjectFilter."), __FUNCTION__, *InObjectFilterClass->GetName()))
+		{
+			ObjectFilterClass = InObjectFilterClass;
+			CacheObjectFilterObject();
+			RequestRebuildList();
+		}
+	}
+
+	FName GetModuleName() const
+	{
+		return ModuleName;
+	}
+
+	// User Collections
+
+	/**
+	 * Add set of objects to a collection in the map, or create a new collection if one does not exist,
+	 */
+	void AddObjectsToCollection(const FName& CollectionName, const TSet<FSoftObjectPath>& ObjectsToAdd) const;
+	void RemoveObjectsFromCollection(const FName& CollectionName, const TSet<FSoftObjectPath>& ObjectsToRemove) const;
+	void RemoveCollection(const FName& CollectionName) const;
+	void ReorderCollection(const FName& CollectionToMoveName, const FName& CollectionInsertBeforeName) const;
+	bool IsObjectInCollection(const FName& CollectionName, const FSoftObjectPath& InObject) const;
+	TSet<FName> GetCollectionsForObject(const FSoftObjectPath& InObject) const;
+	TArray<FName> GetAllCollectionNames() const;
+	FOnObjectMixerCollectionMapChanged& GetOnObjectMixerCollectionMapChanged()
+	{
+		return OnObjectMixerCollectionMapChanged;
+	}
+
+	/**
+	 * Returns the collections selected by the user. If the set is empty, consider "All" collections to be selected.
+	 */
+	const TSet<FName>& GetCurrentCollectionSelection() const;
+
+private:
+
+	TSharedPtr<SObjectMixerEditorMainPanel> MainPanelWidget;
+
+	TSharedPtr<FObjectMixerEditorList> EditorListModel;
+
+	TStrongObjectPtr<UObjectMixerObjectFilter> ObjectFilterPtr;
+
+	/**
+	 * The class used to generate property edit columns
+	 */
+	TSubclassOf<UObjectMixerObjectFilter> ObjectFilterClass;
+
+	/**
+	 * Determines the style of the tree (flat list or hierarchy)
+	 */
+	EObjectMixerTreeViewMode TreeViewMode = EObjectMixerTreeViewMode::FolderObjectSubObject;
+
+	TWeakPtr<FObjectMixerEditorListRow> SoloRow = nullptr;
+
+	FName ModuleName = NAME_None;
+
+	FOnObjectMixerCollectionMapChanged OnObjectMixerCollectionMapChanged;
+};
diff --git a/Engine/Plugins/Editor/ObjectMixer/ObjectMixer/Source/ObjectMixer/Public/Views/MainPanel/SObjectMixerEditorMainPanel.h b/Engine/Plugins/Editor/ObjectMixer/ObjectMixer/Source/ObjectMixer/Public/Views/MainPanel/SObjectMixerEditorMainPanel.h
new file mode 100644
--- /dev/null
+++ b/Engine/Plugins/Editor/ObjectMixer/ObjectMixer/Source/ObjectMixer/Public/Views/MainPanel/SObjectMixerEditorMainPanel.h
@@ -0,0 +1,88 @@
+// Copyright Epic Games, Inc. All Rights Reserved.
+
+#pragma once
+
+#include "ObjectFilter/ObjectMixerEditorObjectFilter.h"
+#include "Views/List/ObjectMixerEditorListFilters/IObjectMixerEditorListFilter.h"
+
+#include "Styling/SlateTypes.h"
+#include "Widgets/SCompoundWidget.h"
+
+class FObjectMixerEditorMainPanel;
+
+class OBJECTMIXEREDITOR_API SObjectMixerEditorMainPanel final : public SCompoundWidget
+{
+
+public:
+
+	SLATE_BEGIN_ARGS(SObjectMixerEditorMainPanel)
+	{}
+
+	SLATE_END_ARGS()
+
+	void Construct(const FArguments& InArgs, const TSharedRef<FObjectMixerEditorMainPanel>& InMainPanel);
+
+	TWeakPtr<FObjectMixerEditorMainPanel> GetMainPanelModel()
+	{
+		return MainPanelModel;
+	}
+
+	FString GetSearchStringFromSearchInputField() const;
+	void SetSearchStringInSearchInputField(const FString InSearchString) const;
+	void ExecuteListViewSearchOnAllRows(const FString& SearchString, const bool bShouldRefreshAfterward = true);
+
+	/**
+	 * Determines the style of the tree (flat list or hierarchy)
+	 */
+	EObjectMixerTreeViewMode GetTreeViewMode();
+	/**
+	 * Determine the style of the tree (flat list or hierarchy)
+	 */
+	void SetTreeViewMode(EObjectMixerTreeViewMode InViewMode);
+
+	void ToggleFilterActive(const FString& FilterName);
+
+	const TArray<TSharedRef<class IObjectMixerEditorListFilter>>& GetShowFilters()
+	{
+		return ShowFilters;	
+	}
+
+	const TSet<FName>& GetCurrentCollectionSelection();
+	
+	void RebuildCollectionSelector();
+
+	void RequestRemoveCollection(const FName& CollectionName);
+	
+	void OnCollectionCheckedStateChanged(ECheckBoxState State, FName CollectionName);
+	ECheckBoxState IsCollectionChecked(FName Section) const;
+
+	virtual ~SObjectMixerEditorMainPanel() override;
+
+private:
+
+	/** A reference to the struct that controls this widget */
+	TWeakPtr<FObjectMixerEditorMainPanel> MainPanelModel;
+
+	TArray<TSharedRef<IObjectMixerEditorListFilter>> ShowFilters;
+
+	TSharedPtr<class SSearchBox> SearchBoxPtr;
+	TSharedPtr<class SComboButton> ViewOptionsComboButton;
+
+	TSharedRef<SWidget> GenerateToolbar();
+	TSharedRef<SWidget> OnGenerateAddObjectButtonMenu() const;
+
+	TSharedRef<SWidget> OnGenerateFilterClassMenu();
+	TSharedRef<SWidget> BuildShowOptionsMenu();
+
+	void OnSearchTextChanged(const FText& Text);
+
+	// User Collections
+	
+	TSharedPtr<class SWrapBox> CollectionSelectorBox;
+	TSet<FName> CurrentCollectionSelection;
+
+	void ResetCurrentCollectionSelection()
+	{
+		CurrentCollectionSelection.Reset();
+	}
+};
diff --git a/Engine/Plugins/Editor/ObjectMixer/ObjectMixer/Source/ObjectMixer/Public/Views/Widgets/ObjectMixerEditorListMenuContext.h b/Engine/Plugins/Editor/ObjectMixer/ObjectMixer/Source/ObjectMixer/Public/Views/Widgets/ObjectMixerEditorListMenuContext.h
new file mode 100644
--- /dev/null
+++ b/Engine/Plugins/Editor/ObjectMixer/ObjectMixer/Source/ObjectMixer/Public/Views/Widgets/ObjectMixerEditorListMenuContext.h
@@ -0,0 +1,44 @@
+// Copyright Epic Games, Inc. All Rights Reserved.
+
+#pragma once
+
+#include "Views/List/ObjectMixerEditorListRow.h"
+
+#include "Types/SlateEnums.h"
+
+#include "ObjectMixerEditorListMenuContext.generated.h"	
+
+UCLASS()
+class OBJECTMIXEREDITOR_API UObjectMixerEditorListMenuContext : public UObject
+{
+
+	GENERATED_BODY()
+	
+public:
+
+	struct FObjectMixerEditorListMenuContextData
+	{
+		TArray<FObjectMixerEditorListRowPtr> SelectedItems;
+		TWeakPtr<class FObjectMixerEditorMainPanel> MainPanelPtr;
+	};
+
+	static TSharedPtr<SWidget> CreateContextMenu(const FObjectMixerEditorListMenuContextData InData);
+	static TSharedPtr<SWidget> BuildContextMenu(const FObjectMixerEditorListMenuContextData& InData);
+	static void RegisterContextMenu();
+
+	FObjectMixerEditorListMenuContextData Data;
+	
+	static FName DefaultContextBaseMenuName;
+
+private:
+
+	static void OnTextCommitted(const FText& InText, ETextCommit::Type InCommitType, UObjectMixerEditorListMenuContext* Context);
+	static void OnTextChanged(const FText& InText, UObjectMixerEditorListMenuContext* Context);
+	
+	static void OnClickCollectionMenuEntry(const FName Key, UObjectMixerEditorListMenuContext* Context);
+	static void AddObjectsToCollection(const FName Key, UObjectMixerEditorListMenuContext* Context);
+	static void RemoveObjectsFromCollection(const FName Key, UObjectMixerEditorListMenuContext* Context);
+	static bool AreAllObjectsInCollection(const FName Key, UObjectMixerEditorListMenuContext* Context);
+
+	TSharedPtr<class SEditableTextBox> EditableText;
+};
diff --git a/Engine/Plugins/Editor/ObjectMixer/ObjectMixer/Source/ObjectMixer/Public/Views/Widgets/SCollectionSelectionButton.h b/Engine/Plugins/Editor/ObjectMixer/ObjectMixer/Source/ObjectMixer/Public/Views/Widgets/SCollectionSelectionButton.h
new file mode 100644
--- /dev/null
+++ b/Engine/Plugins/Editor/ObjectMixer/ObjectMixer/Source/ObjectMixer/Public/Views/Widgets/SCollectionSelectionButton.h
@@ -0,0 +1,52 @@
+// Copyright Epic Games, Inc. All Rights Reserved.
+
+#pragma once
+
+#include "Brushes/SlateRoundedBoxBrush.h"
+#include "Styling/StyleColors.h"
+#include "Widgets/SCompoundWidget.h"
+
+#define LOCTEXT_NAMESPACE "ObjectMixerEditor"
+
+class SCollectionSelectionButton final : public SCompoundWidget
+{
+public:
+	
+	SLATE_BEGIN_ARGS(SCollectionSelectionButton)
+	{}
+
+	SLATE_END_ARGS()
+	
+	void Construct(const FArguments& InArgs, const TSharedRef<class SObjectMixerEditorMainPanel> MainPanelWidget, const FName& InCollectionName);
+
+	virtual FReply OnMouseButtonDown(const FGeometry& MyGeometry, const FPointerEvent& MouseEvent) override;
+
+	virtual FReply OnMouseButtonUp(const FGeometry& MyGeometry, const FPointerEvent& MouseEvent) override;
+
+	virtual FReply OnDragDetected(const FGeometry& MyGeometry, const FPointerEvent& MouseEvent) override;
+
+	virtual void OnDragEnter(const FGeometry& MyGeometry, const FDragDropEvent& DragDropEvent) override;
+
+	virtual void OnDragLeave(const FDragDropEvent& DragDropEvent) override;
+
+	virtual FReply OnDrop(const FGeometry& MyGeometry, const FDragDropEvent& DragDropEvent) override;
+
+	TWeakPtr<SObjectMixerEditorMainPanel> MainPanelPtr;
+	FName CollectionName = NAME_None;
+
+private:
+
+	bool bIsPressed = false;
+	bool bDropIsValid = false;
+
+	FSlateRoundedBoxBrush UncheckedImage = FSlateRoundedBoxBrush(FStyleColors::Header, 4.0f, FStyleColors::Input, 1.0f);
+	FSlateRoundedBoxBrush UncheckedHoveredImage = FSlateRoundedBoxBrush(FStyleColors::Hover, 4.0f, FStyleColors::Input, 1.0f);
+	FSlateRoundedBoxBrush UncheckedPressedImage = FSlateRoundedBoxBrush(FStyleColors::Hover, 4.0f, FStyleColors::Input, 1.0f);
+	FSlateRoundedBoxBrush UncheckedValidDropImage = FSlateRoundedBoxBrush(FStyleColors::Secondary, 4.0f, FStyleColors::Input, 1.0f);
+	FSlateRoundedBoxBrush CheckedImage = FSlateRoundedBoxBrush(FStyleColors::Primary, 4.0f, FStyleColors::Input, 1.0f);
+	FSlateRoundedBoxBrush CheckedHoveredImage = FSlateRoundedBoxBrush(FStyleColors::PrimaryHover, 4.0f, FStyleColors::Input, 1.0f);
+	FSlateRoundedBoxBrush CheckedPressedImage = FSlateRoundedBoxBrush(FStyleColors::PrimaryHover, 4.0f, FStyleColors::Input, 1.0f);
+	FSlateRoundedBoxBrush CheckedValidDropImage = FSlateRoundedBoxBrush(FStyleColors::Secondary, 4.0f, FStyleColors::Input, 1.0f);
+};
+
+#undef LOCTEXT_NAMESPACE
diff --git a/Engine/Source/Editor/LevelEditor/LevelEditor.Build.cs b/Engine/Source/Editor/LevelEditor/LevelEditor.Build.cs
--- a/Engine/Source/Editor/LevelEditor/LevelEditor.Build.cs
+++ b/Engine/Source/Editor/LevelEditor/LevelEditor.Build.cs
@@ -61,8 +61,6 @@ public class LevelEditor : ModuleRules
 				"TypedElementFramework",
 				"TypedElementRuntime",
 				"EngineSettings",
-				"LightMixer",
-				"ObjectMixerEditor",
 				"PropertyEditor",
 				"Kismet",
 				"KismetWidgets",
diff --git a/Engine/Source/Editor/LevelEditor/Private/SLevelEditor.cpp b/Engine/Source/Editor/LevelEditor/Private/SLevelEditor.cpp
--- a/Engine/Source/Editor/LevelEditor/Private/SLevelEditor.cpp
+++ b/Engine/Source/Editor/LevelEditor/Private/SLevelEditor.cpp
@@ -72,8 +72,6 @@
 #include "EditorViewportLayout.h"
 #include "LevelViewportTabContent.h"
 #include "SLevelViewport.h"
-#include "LightMixerModule.h"
-#include "ObjectMixerEditorModule.h"
 
 #define LOCTEXT_NAMESPACE "SLevelEditor"
 
@@ -1486,28 +1484,6 @@ TSharedRef<SWidget> SLevelEditor::RestoreContentArea( const TSharedRef<SDockTab>
 				.SetIcon(EnvironmentLightingViewerIcon);
 		}
 
-		{
-			// These modules will take care of their own tab spawners once loaded.
-			FObjectMixerEditorModule& ObjectMixer = FObjectMixerEditorModule::Get();
-			{
-				FDelegateHandle ObjectMixerComponentEditedHandle =
-					LevelEditorModule.OnComponentsEdited().AddLambda([]()
-					{
-						FObjectMixerEditorModule::Get().RequestRebuildList();
-					});
-				ObjectMixer.AddOnComponentEditedDelegate(ObjectMixerComponentEditedHandle);
-			}
-			FLightMixerModule& LightMixer = FLightMixerModule::Get();
-			{
-				FDelegateHandle LightMixerComponentEditedHandle =
-				   LevelEditorModule.OnComponentsEdited().AddLambda([]()
-				   {
-					   FLightMixerModule::Get().RequestRebuildList();
-				   });
-				LightMixer.AddOnComponentEditedDelegate(LightMixerComponentEditedHandle);
-			}
-		}
-
 		FTabSpawnerEntry& BuildAndSubmitEntry = LevelEditorTabManager->RegisterTabSpawner(LevelEditorTabIds::LevelEditorBuildAndSubmit, FOnSpawnTab::CreateSP<SLevelEditor, FName, FString>(this, &SLevelEditor::SpawnLevelEditorTab, LevelEditorTabIds::LevelEditorBuildAndSubmit, FString()));
 		BuildAndSubmitEntry.SetAutoGenerateMenuEntry(false);
 
diff --git a/Engine/Source/Editor/ObjectMixer/LightMixer/LightMixer.build.cs b/Engine/Source/Editor/ObjectMixer/LightMixer/LightMixer.build.cs
deleted file mode 100644
diff --git a/Engine/Source/Editor/ObjectMixer/LightMixer/Private/LightMixerModule.cpp b/Engine/Source/Editor/ObjectMixer/LightMixer/Private/LightMixerModule.cpp
deleted file mode 100644
diff --git a/Engine/Source/Editor/ObjectMixer/LightMixer/Private/LightMixerStyle.cpp b/Engine/Source/Editor/ObjectMixer/LightMixer/Private/LightMixerStyle.cpp
deleted file mode 100644
diff --git a/Engine/Source/Editor/ObjectMixer/LightMixer/Private/LightMixerStyle.h b/Engine/Source/Editor/ObjectMixer/LightMixer/Private/LightMixerStyle.h
deleted file mode 100644
diff --git a/Engine/Source/Editor/ObjectMixer/LightMixer/Public/LightMixerModule.h b/Engine/Source/Editor/ObjectMixer/LightMixer/Public/LightMixerModule.h
deleted file mode 100644
diff --git a/Engine/Source/Editor/ObjectMixer/LightMixer/Public/LightMixerObjectFilter.h b/Engine/Source/Editor/ObjectMixer/LightMixer/Public/LightMixerObjectFilter.h
deleted file mode 100644
diff --git a/Engine/Source/Editor/ObjectMixer/LightMixer/Public/LightMixerProjectSettings.h b/Engine/Source/Editor/ObjectMixer/LightMixer/Public/LightMixerProjectSettings.h
deleted file mode 100644
diff --git a/Engine/Source/Editor/ObjectMixer/ObjectMixer/ObjectMixerEditor.build.cs b/Engine/Source/Editor/ObjectMixer/ObjectMixer/ObjectMixerEditor.build.cs
deleted file mode 100644
diff --git a/Engine/Source/Editor/ObjectMixer/ObjectMixer/Private/ObjectFilter/ObjectMixerEditorObjectFilter.cpp b/Engine/Source/Editor/ObjectMixer/ObjectMixer/Private/ObjectFilter/ObjectMixerEditorObjectFilter.cpp
deleted file mode 100644
diff --git a/Engine/Source/Editor/ObjectMixer/ObjectMixer/Private/ObjectMixerEditorModule.cpp b/Engine/Source/Editor/ObjectMixer/ObjectMixer/Private/ObjectMixerEditorModule.cpp
deleted file mode 100644
diff --git a/Engine/Source/Editor/ObjectMixer/ObjectMixer/Private/ObjectMixerEditorSerializedData.cpp b/Engine/Source/Editor/ObjectMixer/ObjectMixer/Private/ObjectMixerEditorSerializedData.cpp
deleted file mode 100644
diff --git a/Engine/Source/Editor/ObjectMixer/ObjectMixer/Private/ObjectMixerEditorStyle.cpp b/Engine/Source/Editor/ObjectMixer/ObjectMixer/Private/ObjectMixerEditorStyle.cpp
deleted file mode 100644
diff --git a/Engine/Source/Editor/ObjectMixer/ObjectMixer/Private/ObjectMixerFilterFactory.cpp b/Engine/Source/Editor/ObjectMixer/ObjectMixer/Private/ObjectMixerFilterFactory.cpp
deleted file mode 100644
diff --git a/Engine/Source/Editor/ObjectMixer/ObjectMixer/Private/Views/List/ObjectMixerEditorList.cpp b/Engine/Source/Editor/ObjectMixer/ObjectMixer/Private/Views/List/ObjectMixerEditorList.cpp
deleted file mode 100644
diff --git a/Engine/Source/Editor/ObjectMixer/ObjectMixer/Private/Views/List/ObjectMixerEditorListRow.cpp b/Engine/Source/Editor/ObjectMixer/ObjectMixer/Private/Views/List/ObjectMixerEditorListRow.cpp
deleted file mode 100644
diff --git a/Engine/Source/Editor/ObjectMixer/ObjectMixer/Private/Views/List/SObjectMixerEditorList.cpp b/Engine/Source/Editor/ObjectMixer/ObjectMixer/Private/Views/List/SObjectMixerEditorList.cpp
deleted file mode 100644
diff --git a/Engine/Source/Editor/ObjectMixer/ObjectMixer/Private/Views/List/SObjectMixerEditorListRow.cpp b/Engine/Source/Editor/ObjectMixer/ObjectMixer/Private/Views/List/SObjectMixerEditorListRow.cpp
deleted file mode 100644
diff --git a/Engine/Source/Editor/ObjectMixer/ObjectMixer/Private/Views/MainPanel/ObjectMixerEditorMainPanel.cpp b/Engine/Source/Editor/ObjectMixer/ObjectMixer/Private/Views/MainPanel/ObjectMixerEditorMainPanel.cpp
deleted file mode 100644
diff --git a/Engine/Source/Editor/ObjectMixer/ObjectMixer/Private/Views/MainPanel/SObjectMixerEditorMainPanel.cpp b/Engine/Source/Editor/ObjectMixer/ObjectMixer/Private/Views/MainPanel/SObjectMixerEditorMainPanel.cpp
deleted file mode 100644
diff --git a/Engine/Source/Editor/ObjectMixer/ObjectMixer/Private/Views/Widgets/ObjectMixerEditorListMenuContext.cpp b/Engine/Source/Editor/ObjectMixer/ObjectMixer/Private/Views/Widgets/ObjectMixerEditorListMenuContext.cpp
deleted file mode 100644
diff --git a/Engine/Source/Editor/ObjectMixer/ObjectMixer/Private/Views/Widgets/SCollectionSelectionButton.cpp b/Engine/Source/Editor/ObjectMixer/ObjectMixer/Private/Views/Widgets/SCollectionSelectionButton.cpp
deleted file mode 100644
diff --git a/Engine/Source/Editor/ObjectMixer/ObjectMixer/Private/Views/Widgets/SObjectMixerPlacementAssetMenuEntry.cpp b/Engine/Source/Editor/ObjectMixer/ObjectMixer/Private/Views/Widgets/SObjectMixerPlacementAssetMenuEntry.cpp
deleted file mode 100644
diff --git a/Engine/Source/Editor/ObjectMixer/ObjectMixer/Private/Views/Widgets/SObjectMixerPlacementAssetMenuEntry.h b/Engine/Source/Editor/ObjectMixer/ObjectMixer/Private/Views/Widgets/SObjectMixerPlacementAssetMenuEntry.h
deleted file mode 100644
diff --git a/Engine/Source/Editor/ObjectMixer/ObjectMixer/Public/ObjectFilter/ObjectMixerEditorObjectFilter.h b/Engine/Source/Editor/ObjectMixer/ObjectMixer/Public/ObjectFilter/ObjectMixerEditorObjectFilter.h
deleted file mode 100644
diff --git a/Engine/Source/Editor/ObjectMixer/ObjectMixer/Public/ObjectMixerEditorLog.h b/Engine/Source/Editor/ObjectMixer/ObjectMixer/Public/ObjectMixerEditorLog.h
deleted file mode 100644
diff --git a/Engine/Source/Editor/ObjectMixer/ObjectMixer/Public/ObjectMixerEditorModule.h b/Engine/Source/Editor/ObjectMixer/ObjectMixer/Public/ObjectMixerEditorModule.h
deleted file mode 100644
diff --git a/Engine/Source/Editor/ObjectMixer/ObjectMixer/Public/ObjectMixerEditorProjectSettings.h b/Engine/Source/Editor/ObjectMixer/ObjectMixer/Public/ObjectMixerEditorProjectSettings.h
deleted file mode 100644
diff --git a/Engine/Source/Editor/ObjectMixer/ObjectMixer/Public/ObjectMixerEditorSerializedData.h b/Engine/Source/Editor/ObjectMixer/ObjectMixer/Public/ObjectMixerEditorSerializedData.h
deleted file mode 100644
diff --git a/Engine/Source/Editor/ObjectMixer/ObjectMixer/Public/ObjectMixerEditorStyle.h b/Engine/Source/Editor/ObjectMixer/ObjectMixer/Public/ObjectMixerEditorStyle.h
deleted file mode 100644
diff --git a/Engine/Source/Editor/ObjectMixer/ObjectMixer/Public/ObjectMixerFilterFactory.h b/Engine/Source/Editor/ObjectMixer/ObjectMixer/Public/ObjectMixerFilterFactory.h
deleted file mode 100644
diff --git a/Engine/Source/Editor/ObjectMixer/ObjectMixer/Public/Views/List/ObjectMixerEditorList.h b/Engine/Source/Editor/ObjectMixer/ObjectMixer/Public/Views/List/ObjectMixerEditorList.h
deleted file mode 100644
diff --git a/Engine/Source/Editor/ObjectMixer/ObjectMixer/Public/Views/List/ObjectMixerEditorListFilters/IObjectMixerEditorListFilter.h b/Engine/Source/Editor/ObjectMixer/ObjectMixer/Public/Views/List/ObjectMixerEditorListFilters/IObjectMixerEditorListFilter.h
deleted file mode 100644
diff --git a/Engine/Source/Editor/ObjectMixer/ObjectMixer/Public/Views/List/ObjectMixerEditorListFilters/ObjectMixerEditorListFilter_Collection.h b/Engine/Source/Editor/ObjectMixer/ObjectMixer/Public/Views/List/ObjectMixerEditorListFilters/ObjectMixerEditorListFilter_Collection.h
deleted file mode 100644
diff --git a/Engine/Source/Editor/ObjectMixer/ObjectMixer/Public/Views/List/ObjectMixerEditorListRow.h b/Engine/Source/Editor/ObjectMixer/ObjectMixer/Public/Views/List/ObjectMixerEditorListRow.h
deleted file mode 100644
diff --git a/Engine/Source/Editor/ObjectMixer/ObjectMixer/Public/Views/List/SObjectMixerEditorList.h b/Engine/Source/Editor/ObjectMixer/ObjectMixer/Public/Views/List/SObjectMixerEditorList.h
deleted file mode 100644
diff --git a/Engine/Source/Editor/ObjectMixer/ObjectMixer/Public/Views/List/SObjectMixerEditorListRow.h b/Engine/Source/Editor/ObjectMixer/ObjectMixer/Public/Views/List/SObjectMixerEditorListRow.h
deleted file mode 100644
diff --git a/Engine/Source/Editor/ObjectMixer/ObjectMixer/Public/Views/MainPanel/ObjectMixerEditorMainPanel.h b/Engine/Source/Editor/ObjectMixer/ObjectMixer/Public/Views/MainPanel/ObjectMixerEditorMainPanel.h
deleted file mode 100644
diff --git a/Engine/Source/Editor/ObjectMixer/ObjectMixer/Public/Views/MainPanel/SObjectMixerEditorMainPanel.h b/Engine/Source/Editor/ObjectMixer/ObjectMixer/Public/Views/MainPanel/SObjectMixerEditorMainPanel.h
deleted file mode 100644
diff --git a/Engine/Source/Editor/ObjectMixer/ObjectMixer/Public/Views/Widgets/ObjectMixerEditorListMenuContext.h b/Engine/Source/Editor/ObjectMixer/ObjectMixer/Public/Views/Widgets/ObjectMixerEditorListMenuContext.h
deleted file mode 100644
diff --git a/Engine/Source/Editor/ObjectMixer/ObjectMixer/Public/Views/Widgets/SCollectionSelectionButton.h b/Engine/Source/Editor/ObjectMixer/ObjectMixer/Public/Views/Widgets/SCollectionSelectionButton.h
deleted file mode 100644
