Objet Mixer:

-Move "All" collection to first slot in collection selector
-Add some logic to ensure drag and drop operations are valid with respect to collection selection buttons
-Clean up Slate code for readability

diff --git a/Engine/Plugins/Editor/ObjectMixer/ObjectMixer/Source/ObjectMixer/Private/Views/Widgets/SCollectionSelectionButton.cpp b/Engine/Plugins/Editor/ObjectMixer/ObjectMixer/Source/ObjectMixer/Private/Views/Widgets/SCollectionSelectionButton.cpp
--- a/Engine/Plugins/Editor/ObjectMixer/ObjectMixer/Source/ObjectMixer/Private/Views/Widgets/SCollectionSelectionButton.cpp
+++ b/Engine/Plugins/Editor/ObjectMixer/ObjectMixer/Source/ObjectMixer/Private/Views/Widgets/SCollectionSelectionButton.cpp
@@ -50,42 +50,7 @@ void SCollectionSelectionButton::Construct(const FArguments& InArgs, const TShar
 	[
 		SNew(SBorder)
 		.Padding(FMargin(16, 4))
-		.BorderImage_Lambda([this]()
-		{
-			if (MainPanelPtr.Pin()->IsCollectionChecked(CollectionName) == ECheckBoxState::Checked)
-			{
-				if (bIsPressed)
-				{
-					return &CheckedPressedImage;
-				}
-				else if (bDropIsValid)
-				{
-					return &CheckedValidDropImage;
-				}
-				else if (IsHovered())
-				{
-					return &CheckedHoveredImage;
-				}
-
-				return &CheckedImage;
-			}
-
-			if (bIsPressed)
-			{
-				return &UncheckedPressedImage;
-			}
-			else if (bDropIsValid)
-			{
-				return &UncheckedValidDropImage;
-			}
-			else if (IsHovered())
-			{
-				return &UncheckedHoveredImage;
-			}
-
-			return &UncheckedImage;
-			
-		})
+		.BorderImage(this, &SCollectionSelectionButton::GetBorderBrush)
 		[
 			SNew(STextBlock)
 			.TextStyle(FAppStyle::Get(), "SmallText")
@@ -168,7 +133,7 @@ void SCollectionSelectionButton::OnDragEnter(const FGeometry& MyGeometry, const
 	if (TSharedPtr<FCollectionSelectionButtonDragDropOp> OperationFromCollection =
 		DragDropEvent.GetOperationAs<FCollectionSelectionButtonDragDropOp>())
 	{
-		if (CollectionName != NAME_None && CollectionName != OperationFromCollection->DraggedItem)
+		if (CollectionName != "All" && CollectionName != NAME_None && CollectionName != OperationFromCollection->DraggedItem)
 		{
 			OperationFromCollection->SetToolTip(
 			   FText::Format(
@@ -207,32 +172,75 @@ FReply SCollectionSelectionButton::OnDrop(const FGeometry& MyGeometry, const FDr
 	if (TSharedPtr<FObjectMixerListRowDragDropOp> OperationFromRow =
 		DragDropEvent.GetOperationAs<FObjectMixerListRowDragDropOp>())
 	{
-		if (const TSharedPtr<FObjectMixerEditorMainPanel> MainPanelModel = MainPanelPtr.Pin()->GetMainPanelModel().Pin())
+		if (CollectionName != "All" && CollectionName != NAME_None)
 		{
-			TSet<FSoftObjectPath> ObjectPaths;
-
-			for (const TSharedPtr<FObjectMixerEditorListRow>& Item : OperationFromRow->DraggedItems)
+			if (const TSharedPtr<FObjectMixerEditorMainPanel> MainPanelModel = MainPanelPtr.Pin()->GetMainPanelModel().Pin())
 			{
-				if (const UObject* Object = Item->GetObject())
+				TSet<FSoftObjectPath> ObjectPaths;
+
+				for (const TSharedPtr<FObjectMixerEditorListRow>& Item : OperationFromRow->DraggedItems)
 				{
-					ObjectPaths.Add(Object);
+					if (const UObject* Object = Item->GetObject())
+					{
+						ObjectPaths.Add(Object);
+					}
 				}
-			}
 	
-			MainPanelModel->AddObjectsToCollection(CollectionName, ObjectPaths);
+				MainPanelModel->AddObjectsToCollection(CollectionName, ObjectPaths);
+			}
 		}
 	}
 
 	if (TSharedPtr<FCollectionSelectionButtonDragDropOp> OperationFromCollection =
 		DragDropEvent.GetOperationAs<FCollectionSelectionButtonDragDropOp>())
 	{
-		if (const TSharedPtr<FObjectMixerEditorMainPanel> MainPanelModel = MainPanelPtr.Pin()->GetMainPanelModel().Pin())
+		if (CollectionName != "All" && CollectionName != NAME_None && CollectionName != OperationFromCollection->DraggedItem)
 		{
-			MainPanelModel->ReorderCollection(OperationFromCollection->DraggedItem, CollectionName);
+			if (const TSharedPtr<FObjectMixerEditorMainPanel> MainPanelModel = MainPanelPtr.Pin()->GetMainPanelModel().Pin())
+			{
+				MainPanelModel->ReorderCollection(OperationFromCollection->DraggedItem, CollectionName);
+			}
 		}
 	}
 		
 	return FReply::Handled();
 }
 
+const FSlateBrush* SCollectionSelectionButton::GetBorderBrush() const
+{
+	if (MainPanelPtr.Pin()->IsCollectionChecked(CollectionName) == ECheckBoxState::Checked)
+	{
+		if (bIsPressed)
+		{
+			return &CheckedPressedImage;
+		}
+		else if (bDropIsValid)
+		{
+			return &CheckedValidDropImage;
+		}
+		else if (IsHovered())
+		{
+			return &CheckedHoveredImage;
+		}
+
+		return &CheckedImage;
+	}
+
+	if (bIsPressed)
+	{
+		return &UncheckedPressedImage;
+	}
+	else if (bDropIsValid)
+	{
+		return &UncheckedValidDropImage;
+	}
+	else if (IsHovered())
+	{
+		return &UncheckedHoveredImage;
+	}
+
+	return &UncheckedImage;
+			
+}
+
 #undef LOCTEXT_NAMESPACE
diff --git a/Engine/Plugins/Editor/ObjectMixer/ObjectMixer/Source/ObjectMixer/Public/Views/Widgets/SCollectionSelectionButton.h b/Engine/Plugins/Editor/ObjectMixer/ObjectMixer/Source/ObjectMixer/Public/Views/Widgets/SCollectionSelectionButton.h
--- a/Engine/Plugins/Editor/ObjectMixer/ObjectMixer/Source/ObjectMixer/Public/Views/Widgets/SCollectionSelectionButton.h
+++ b/Engine/Plugins/Editor/ObjectMixer/ObjectMixer/Source/ObjectMixer/Public/Views/Widgets/SCollectionSelectionButton.h
@@ -31,6 +31,8 @@ public:
 
 	virtual FReply OnDrop(const FGeometry& MyGeometry, const FDragDropEvent& DragDropEvent) override;
 
+	const FSlateBrush* GetBorderBrush() const;
+
 	TWeakPtr<SObjectMixerEditorMainPanel> MainPanelPtr;
 	FName CollectionName = NAME_None;
 
