Object Mixer/Light Mixer

-Fix for: Object Mixer plugins invalidate existing project files (remove "Installed" from plugin descriptors)

diff --git a/Engine/Plugins/Editor/ObjectMixer/LightMixer/LightMixer.uplugin b/Engine/Plugins/Editor/ObjectMixer/LightMixer/LightMixer.uplugin
--- a/Engine/Plugins/Editor/ObjectMixer/LightMixer/LightMixer.uplugin
+++ b/Engine/Plugins/Editor/ObjectMixer/LightMixer/LightMixer.uplugin
@@ -14,7 +14,6 @@
 	"CanContainContent": true,
 	"IsBetaVersion": true,
 	"IsExperimentalVersion": false,
-	"Installed": true,
 	"Hidden": true,
 	"Plugins": [
 		{
diff --git a/Engine/Plugins/Editor/ObjectMixer/ObjectMixer/ObjectMixer.uplugin b/Engine/Plugins/Editor/ObjectMixer/ObjectMixer/ObjectMixer.uplugin
--- a/Engine/Plugins/Editor/ObjectMixer/ObjectMixer/ObjectMixer.uplugin
+++ b/Engine/Plugins/Editor/ObjectMixer/ObjectMixer/ObjectMixer.uplugin
@@ -14,7 +14,6 @@
 	"CanContainContent": true,
 	"IsBetaVersion": true,
 	"IsExperimentalVersion": false,
-	"Installed": true,
 	"Hidden": true,
 	"Modules": [
 		{
