Remove accidental `#pragma once` in cpp files

diff --git a/Engine/Plugins/Editor/ObjectMixer/ObjectMixer/Source/ObjectMixer/Private/Views/Widgets/SCollectionSelectionButton.cpp b/Engine/Plugins/Editor/ObjectMixer/ObjectMixer/Source/ObjectMixer/Private/Views/Widgets/SCollectionSelectionButton.cpp
--- a/Engine/Plugins/Editor/ObjectMixer/ObjectMixer/Source/ObjectMixer/Private/Views/Widgets/SCollectionSelectionButton.cpp
+++ b/Engine/Plugins/Editor/ObjectMixer/ObjectMixer/Source/ObjectMixer/Private/Views/Widgets/SCollectionSelectionButton.cpp
@@ -1,7 +1,5 @@
 // Copyright Epic Games, Inc. All Rights Reserved.
 
-#pragma once
-
 #include "Views/Widgets/SCollectionSelectionButton.h"
 
 #include "DragAndDrop/DecoratedDragDropOp.h"
diff --git a/Engine/Plugins/Experimental/UVEditor/Source/UVEditor/Private/UVEditorAssetEditorImpl.cpp b/Engine/Plugins/Experimental/UVEditor/Source/UVEditor/Private/UVEditorAssetEditorImpl.cpp
--- a/Engine/Plugins/Experimental/UVEditor/Source/UVEditor/Private/UVEditorAssetEditorImpl.cpp
+++ b/Engine/Plugins/Experimental/UVEditor/Source/UVEditor/Private/UVEditorAssetEditorImpl.cpp
@@ -1,7 +1,5 @@
 // Copyright Epic Games, Inc. All Rights Reserved.
 
-#pragma once
-
 #include "UVEditorAssetEditorImpl.h"
 
 #include "CoreMinimal.h"
