Object Mixer:

-Hybrid View

diff --git a/Engine/Plugins/Editor/ObjectMixer/ObjectMixer/Source/ObjectMixer/Private/ObjectFilter/ObjectMixerEditorObjectFilter.cpp b/Engine/Plugins/Editor/ObjectMixer/ObjectMixer/Source/ObjectMixer/Private/ObjectFilter/ObjectMixerEditorObjectFilter.cpp
--- a/Engine/Plugins/Editor/ObjectMixer/ObjectMixer/Source/ObjectMixer/Private/ObjectFilter/ObjectMixerEditorObjectFilter.cpp
+++ b/Engine/Plugins/Editor/ObjectMixer/ObjectMixer/Source/ObjectMixer/Private/ObjectFilter/ObjectMixerEditorObjectFilter.cpp
@@ -9,7 +9,7 @@ FText UObjectMixerObjectFilter::GetRowDisplayName(UObject* InObject, EObjectMixe
 {
 	if (IsValid(InObject))
 	{
-		if (InViewMode == EObjectMixerTreeViewMode::Flat || InViewMode == EObjectMixerTreeViewMode::Folder)
+		if (InViewMode == EObjectMixerTreeViewMode::NoFolders || InViewMode == EObjectMixerTreeViewMode::Folders)
 		{
 			if (const AActor* Outer = InObject->GetTypedOuter<AActor>())
 			{
diff --git a/Engine/Plugins/Editor/ObjectMixer/ObjectMixer/Source/ObjectMixer/Private/Views/List/ObjectMixerEditorListRow.cpp b/Engine/Plugins/Editor/ObjectMixer/ObjectMixer/Source/ObjectMixer/Private/Views/List/ObjectMixerEditorListRow.cpp
--- a/Engine/Plugins/Editor/ObjectMixer/ObjectMixer/Source/ObjectMixer/Private/Views/List/ObjectMixerEditorListRow.cpp
+++ b/Engine/Plugins/Editor/ObjectMixer/ObjectMixer/Source/ObjectMixer/Private/Views/List/ObjectMixerEditorListRow.cpp
@@ -92,6 +92,17 @@ FObjectMixerEditorListRow::EObjectMixerEditorListRowType FObjectMixerEditorListR
 	return RowType;
 }
 
+void FObjectMixerEditorListRow::SetRowType(EObjectMixerEditorListRowType InNewRowType)
+{
+	RowType = InNewRowType;
+}
+
+bool FObjectMixerEditorListRow::IsHybridRow() const
+{
+	return GetRowType() == FObjectMixerEditorListRow::ContainerObject &&
+			GetChildCount() == 1 && GetChildRows()[0]->GetChildCount() == 0;
+}
+
 int32 FObjectMixerEditorListRow::GetSortOrder() const
 {
 	return SortOrder;
diff --git a/Engine/Plugins/Editor/ObjectMixer/ObjectMixer/Source/ObjectMixer/Private/Views/List/SObjectMixerEditorList.cpp b/Engine/Plugins/Editor/ObjectMixer/ObjectMixer/Source/ObjectMixer/Private/Views/List/SObjectMixerEditorList.cpp
--- a/Engine/Plugins/Editor/ObjectMixer/ObjectMixer/Source/ObjectMixer/Private/Views/List/SObjectMixerEditorList.cpp
+++ b/Engine/Plugins/Editor/ObjectMixer/ObjectMixer/Source/ObjectMixer/Private/Views/List/SObjectMixerEditorList.cpp
@@ -24,7 +24,6 @@
 #include "Editor/UnrealEdEngine.h"
 #include "PlacementMode/Public/IPlacementModeModule.h"
 #include "Styling/StyleColors.h"
-#include "UObject/UObjectIterator.h"
 #include "Views/Widgets/ObjectMixerEditorListMenuContext.h"
 #include "Widgets/Layout/SScrollBox.h"
 #include "Widgets/Layout/SWidgetSwitcher.h"
@@ -937,180 +936,177 @@ void SObjectMixerEditorList::BuildPerformanceCacheAndGenerateHeaderIfNeeded()
 	}
 }
 
-void SObjectMixerEditorList::GenerateTreeView()
+bool DoesWorldObjectHaveAcceptableClass(const UObject* Object, const TSet<UClass*>& ObjectClassesToFilterCache)
 {
-	check(ListModelPtr.IsValid());
-	
-	if (!ensure(TreeViewPtr.IsValid()))
+	for (UClass* Class : ObjectClassesToFilterCache)
 	{
-		return;
+		if (Object->IsA(Class))
+		{
+			return true;
+		}
 	}
 
-	CacheTreeState();
-	
-	FlushMemory(true);
-
-	BuildPerformanceCacheAndGenerateHeaderIfNeeded();
-
-	check(GEditor);
-	const UWorld* EditorWorld = GEditor->GetEditorWorldContext().World();
+	return false;
+}
 
-	// Find valid matching objects
-	TArray<UObject*> MatchingObjects;
-	for (TObjectIterator<UObject> ObjectIterator; ObjectIterator; ++ObjectIterator)
+void CreateOuterRowsForTopLevelRow(
+	UObject* InObject,
+	FObjectMixerEditorListRowPtr& TopLevelRow,
+	const TSet<UObject*>& AllMatchingObjects,
+	TMap<UObject*, FObjectMixerEditorListRowPtr>& ObjectsWithRowCreated,
+	const TSharedRef<SObjectMixerEditorList>& InListView)
+{
+	if (InObject)
 	{
-		if (UObject* Object = *ObjectIterator; IsValid(Object))
+		if (InObject->IsA(ULevel::StaticClass()) ||
+			InObject->IsA(UWorld::StaticClass()) ||
+			InObject->IsA(UPackage::StaticClass()))
 		{
-			if (UWorld* ObjectWorld = Object->GetWorld(); ObjectWorld == EditorWorld)
-			{
-				bool bIsAcceptableClass = false;
+			return;
+		}
 
-				for (UClass* Class : ObjectClassesToFilterCache)
-				{
-					if (Object->IsA(Class))
-					{
-						bIsAcceptableClass = true;
-						break;
-					}
-				}
+		FObjectMixerEditorListRowPtr OuterRow = nullptr;
+			
+		if (const FObjectMixerEditorListRowPtr* Match = ObjectsWithRowCreated.Find(InObject))
+		{
+			OuterRow = *Match;
+		}
+		else
+		{				
+			OuterRow = MakeShared<FObjectMixerEditorListRow>(
+				InObject,
+				FObjectMixerEditorListRow::MatchingObject,
+				InListView);
+				
+			check(OuterRow.IsValid());
+			ObjectsWithRowCreated.Add(InObject, OuterRow);
+		}
 
-				if (bIsAcceptableClass)
-				{
-					MatchingObjects.Add(Object);
-				}
+		if (TopLevelRow.IsValid())
+		{
+			OuterRow->AddToChildRows(TopLevelRow);
+				
+			const bool bMatchingContainer = AllMatchingObjects.Contains(InObject);
+				
+			OuterRow->SetRowType(
+				bMatchingContainer ?
+				FObjectMixerEditorListRow::MatchingContainerObject :
+				FObjectMixerEditorListRow::ContainerObject
+			);
+		}
+		TopLevelRow = OuterRow;
+			
+		UObject* Outer = InObject->GetOuter();
+
+		if (!Outer)
+		{
+			if (const AActor* AsActor = Cast<AActor>(InObject))
+			{
+				Outer = AsActor->GetAttachParentActor();
 			}
 		}
-	}
 
-	// A quick lookup for objects that already exist in the list.
-	// Helpful to avoid double-generating rows when considering parent->child hierarchy.
-	TMap<UObject*, TSharedRef<FObjectMixerEditorListRow>> CreatedObjectMap;
-	TMap<FName, TSharedRef<FObjectMixerEditorListRow>> FolderMap;
-	for (UObject* Object : MatchingObjects)
-	{
-		if (!CreatedObjectMap.Contains(Object)) // Ensure we don't double-generate container objects
+		// Recurse with Outer object
+		if (Outer)
 		{
-			TSharedRef<FObjectMixerEditorListRow> TopLevelRow = MakeShared<FObjectMixerEditorListRow>(
-				Object, FObjectMixerEditorListRow::MatchingObject, SharedThis(this));
+			CreateOuterRowsForTopLevelRow(
+			   Outer, TopLevelRow, AllMatchingObjects,
+			   ObjectsWithRowCreated, InListView
+		   );
+		}
+	}
+}
 
-			CreatedObjectMap.Add(Object, TopLevelRow);
+void CreateFolderRowsForTopLevelRow(
+	FObjectMixerEditorListRowPtr& TopLevelRow,
+	TMap<FName, FObjectMixerEditorListRowPtr>& FolderMap,
+	const TSharedRef<SObjectMixerEditorList>& InListView)
+{
+	if (const TObjectPtr<AActor> AsActor = Cast<AActor>(TopLevelRow->GetObject()))
+	{
+		FFolder BaseActorFolder = AsActor->GetFolder();
 
-			// If the view is not in flat mode, we need to consider the hierarchy of outliner folders/attach parents as desired
-			if (GetTreeViewMode() != EObjectMixerTreeViewMode::Flat)
+		while (!BaseActorFolder.IsNone())
+		{
+			TSharedPtr<FObjectMixerEditorListRow> FolderRow;
+			
+			if (const FObjectMixerEditorListRowPtr* Match = FolderMap.Find(BaseActorFolder.GetPath()))
 			{
-				AActor* BaseActor = Cast<AActor>(Object);
-				
-				if (!BaseActor)
-				{
-					if (Object->IsA(UActorComponent::StaticClass()))
-					{
-						BaseActor = Object->GetTypedOuter<AActor>();
+				FolderRow = *Match;
+			}
+			else
+			{
+				FolderRow =
+					MakeShared<FObjectMixerEditorListRow>(
+						nullptr, FObjectMixerEditorListRow::Folder, InListView,
+						FText::FromName(BaseActorFolder.GetLeafName()));
 
-						// If it's not flat or folder view mode, we need to find or create the container object for the actor that owns the matching component
-						if (GetTreeViewMode() != EObjectMixerTreeViewMode::Folder)
-						{
-							TSharedPtr<FObjectMixerEditorListRow> OwningActorRow;
+				FolderMap.Add(BaseActorFolder.GetPath(), FolderRow.ToSharedRef());
+			}
 
-							if (BaseActor)
-							{
-								if (const TSharedRef<FObjectMixerEditorListRow>* Match = CreatedObjectMap.Find(BaseActor))
-								{
-									OwningActorRow = *Match;
-								}
-								else
-								{
-									OwningActorRow = MakeShared<FObjectMixerEditorListRow>(
-										BaseActor, FObjectMixerEditorListRow::ContainerObject, SharedThis(this));
-
-									CreatedObjectMap.Add(BaseActor, OwningActorRow.ToSharedRef());
-								}
-
-								if (OwningActorRow)
-								{
-									OwningActorRow->AddToChildRows(TopLevelRow);
-													
-									TopLevelRow = OwningActorRow.ToSharedRef();
-								}
-							}
-						}
-					}
-				}
+			if (FolderRow)
+			{
+				FolderRow->AddToChildRows(TopLevelRow);
+									
+				TopLevelRow = FolderRow.ToSharedRef();
+			}
 
-				if (BaseActor)
-				{
-					while (AActor* AttachParent = BaseActor->GetAttachParentActor())
-					{
-						// Make a row for each attach parent up the chain until we reach the top if not in flat/folder mode
-						if (GetTreeViewMode() != EObjectMixerTreeViewMode::Folder)
-						{
-							TSharedPtr<FObjectMixerEditorListRow> OwningActorRow;
+			BaseActorFolder = BaseActorFolder.GetParent();
+		}
+	}
+}
 
-							if (const TSharedRef<FObjectMixerEditorListRow>* Match = CreatedObjectMap.Find(AttachParent))
-							{
-								OwningActorRow = *Match;
-							}
-							else
-							{
-								OwningActorRow = MakeShared<FObjectMixerEditorListRow>(
-									AttachParent, FObjectMixerEditorListRow::ContainerObject, SharedThis(this));
+void SObjectMixerEditorList::GenerateTreeView()
+{
+	check(ListModelPtr.IsValid());
+	
+	if (!ensure(TreeViewPtr.IsValid()))
+	{
+		return;
+	}
 
-								CreatedObjectMap.Add(AttachParent, OwningActorRow.ToSharedRef());
-							}
+	CacheTreeState();
+	
+	FlushMemory(true);
 
-							if (OwningActorRow)
-							{
-								OwningActorRow->AddToChildRows(TopLevelRow);
-												
-								TopLevelRow = OwningActorRow.ToSharedRef();
-							}
-						}
+	BuildPerformanceCacheAndGenerateHeaderIfNeeded();
 
-						BaseActor = AttachParent;
-					}
+	check(GEditor);
+	const UWorld* EditorWorld = GEditor->GetEditorWorldContext().World();
+	
+	TSet<UObject*> AllMatchingObjects;
+	ForEachObjectWithOuter(EditorWorld,
+		[this, &AllMatchingObjects](UObject* Object)
+	{
+		if (DoesWorldObjectHaveAcceptableClass(Object, ObjectClassesToFilterCache))
+		{
+			AllMatchingObjects.Add(Object);
+		}
+	});
 
-					// Now consider folder hierarchy for the base actor if desired
-					if (GetTreeViewMode() == EObjectMixerTreeViewMode::Folder ||
-						GetTreeViewMode() == EObjectMixerTreeViewMode::FolderObjectSubObject)
-					{
-						if (BaseActor)
-						{
-							FFolder BaseActorFolder = BaseActor->GetFolder();
+	TMap<UObject*, FObjectMixerEditorListRowPtr> ObjectsWithRowCreated;
+	TMap<FName, FObjectMixerEditorListRowPtr> FolderMap;
+	for (UObject* Object : AllMatchingObjects)
+	{
+		FObjectMixerEditorListRowPtr TopLevelRow = nullptr;
 
-							while (!BaseActorFolder.IsNone())
-							{
-								TSharedPtr<FObjectMixerEditorListRow> FolderRow;
-							
-								if (const TSharedRef<FObjectMixerEditorListRow>* Match = FolderMap.Find(BaseActorFolder.GetPath()))
-								{
-									FolderRow = *Match;
-								}
-								else
-								{
-									FolderRow =
-										MakeShared<FObjectMixerEditorListRow>(
-											nullptr, FObjectMixerEditorListRow::Folder, SharedThis(this),
-											FText::FromName(BaseActorFolder.GetLeafName()));
-
-									FolderMap.Add(BaseActorFolder.GetPath(), FolderRow.ToSharedRef());
-								}
-
-								if (FolderRow)
-								{
-									FolderRow->AddToChildRows(TopLevelRow);
-													
-									TopLevelRow = FolderRow.ToSharedRef();
-								}
-
-								BaseActorFolder = BaseActorFolder.GetParent();
-							}
-						}
-					}
-				}
-			}
+		UObject* NextOuter = Object;
 
-			TreeViewRootObjects.AddUnique(TopLevelRow);
+		CreateOuterRowsForTopLevelRow(
+			NextOuter, TopLevelRow, AllMatchingObjects,
+			ObjectsWithRowCreated, SharedThis(this)
+		);
+		
+		// Now consider folder hierarchy for the top level row's object if desired
+		if (GetTreeViewMode() == EObjectMixerTreeViewMode::Folders)
+		{
+			CreateFolderRowsForTopLevelRow(TopLevelRow, FolderMap, SharedThis(this));
 		}
+
+		TreeViewRootObjects.AddUnique(TopLevelRow);
 	}
+	
+	ObjectsWithRowCreated.Empty();
 
 	TreeViewRootObjects.StableSort(SortByTypeThenName);
 
@@ -1183,12 +1179,15 @@ void SObjectMixerEditorList::OnGetRowChildren(FObjectMixerEditorListRowPtr Row,
 {
 	if (Row.IsValid())
 	{
-		OutChildren = Row->GetChildRows();
-
-		if (Row->GetShouldExpandAllChildren())
+		if (!Row->IsHybridRow())
 		{
-			SetChildExpansionRecursively(Row, true);
-			Row->SetShouldExpandAllChildren(false);
+			OutChildren = Row->GetChildRows();
+
+			if (Row->GetShouldExpandAllChildren())
+			{
+				SetChildExpansionRecursively(Row, true);
+				Row->SetShouldExpandAllChildren(false);
+			}
 		}
 	}
 }
diff --git a/Engine/Plugins/Editor/ObjectMixer/ObjectMixer/Source/ObjectMixer/Private/Views/List/SObjectMixerEditorListRow.cpp b/Engine/Plugins/Editor/ObjectMixer/ObjectMixer/Source/ObjectMixer/Private/Views/List/SObjectMixerEditorListRow.cpp
--- a/Engine/Plugins/Editor/ObjectMixer/ObjectMixer/Source/ObjectMixer/Private/Views/List/SObjectMixerEditorListRow.cpp
+++ b/Engine/Plugins/Editor/ObjectMixer/ObjectMixer/Source/ObjectMixer/Private/Views/List/SObjectMixerEditorListRow.cpp
@@ -37,7 +37,7 @@ void SObjectMixerEditorListRow::Construct(
 	check(InRow.IsValid());
 
 	Item = InRow;
-	const FObjectMixerEditorListRowPtr PinnedItem = Item.Pin();
+	bIsHybridRow = Item.Pin()->IsHybridRow();
 
 	SMultiColumnTableRow<FObjectMixerEditorListRowPtr>::Construct(
 		FSuperRowType::FArguments()
@@ -62,7 +62,8 @@ void SObjectMixerEditorListRow::Construct(
 
 TSharedRef<SWidget> SObjectMixerEditorListRow::GenerateWidgetForColumn(const FName& InColumnName)
 {
-	const FObjectMixerEditorListRowPtr PinnedItem = Item.Pin();
+	const FObjectMixerEditorListRowPtr PinnedItem =
+		bIsHybridRow ? Item.Pin()->GetChildRows()[0] : Item.Pin();
 
 	if (const TSharedPtr<SWidget> CellWidget = GenerateCells(InColumnName, PinnedItem))
 	{
diff --git a/Engine/Plugins/Editor/ObjectMixer/ObjectMixer/Source/ObjectMixer/Public/ObjectFilter/ObjectMixerEditorObjectFilter.h b/Engine/Plugins/Editor/ObjectMixer/ObjectMixer/Source/ObjectMixer/Public/ObjectFilter/ObjectMixerEditorObjectFilter.h
--- a/Engine/Plugins/Editor/ObjectMixer/ObjectMixer/Source/ObjectMixer/Public/ObjectFilter/ObjectMixerEditorObjectFilter.h
+++ b/Engine/Plugins/Editor/ObjectMixer/ObjectMixer/Source/ObjectMixer/Public/ObjectFilter/ObjectMixerEditorObjectFilter.h
@@ -24,17 +24,15 @@ enum class EObjectMixerInheritanceInclusionOptions : uint8
 UENUM(BlueprintType)
 enum class EObjectMixerTreeViewMode : uint8
 {
-	Flat, // Show all matching objects in a flat list without any hierarchy
-	Folder, // Display objects in a folder hierarchy but without Object->SubObject hierarchy
-	ObjectSubObject, // Display objects as SubObjects of BaseObjects (usually actors) without considering folder hierarchy
-	FolderObjectSubObject // Display objects in a full hierarchy (folder->BaseObject->SubObject)
+	NoFolders, // Show all matching objects in a flat list without any hierarchy
+	Folders // Display objects in a folder hierarchy but without Object->SubObject hierarchy
 };
 
 /**
  * Native class for filtering object types to Object Mixer.
  * Native C++ classes should inherit directly from this class.
  */
-UCLASS(Abstract, Blueprintable)
+UCLASS(Abstract, Blueprintable) 
 class OBJECTMIXEREDITOR_API UObjectMixerObjectFilter : public UObject
 {
 	GENERATED_BODY()
@@ -192,7 +190,7 @@ public:
 	 */
 	UFUNCTION(BlueprintNativeEvent, BlueprintCallable, Category = "Object Mixer")
 	FText GetRowDisplayName(
-		UObject* InObject, EObjectMixerTreeViewMode InViewMode = EObjectMixerTreeViewMode::FolderObjectSubObject) const override;
+		UObject* InObject, EObjectMixerTreeViewMode InViewMode = EObjectMixerTreeViewMode::Folders) const override;
 	
 	FText GetRowDisplayName_Implementation(UObject* InObject, EObjectMixerTreeViewMode InViewMode) const
 	{
@@ -207,7 +205,7 @@ public:
 	 */
 	UFUNCTION(BlueprintNativeEvent, BlueprintCallable, Category = "Object Mixer")
 	bool GetRowEditorVisibility(
-		UObject* InObject, EObjectMixerTreeViewMode InViewMode = EObjectMixerTreeViewMode::FolderObjectSubObject) const override;
+		UObject* InObject, EObjectMixerTreeViewMode InViewMode = EObjectMixerTreeViewMode::Folders) const override;
 
 	bool GetRowEditorVisibility_Implementation(UObject* InObject, EObjectMixerTreeViewMode InViewMode) const
 	{
@@ -222,7 +220,7 @@ public:
 	 */
 	UFUNCTION(BlueprintNativeEvent, BlueprintCallable, Category = "Object Mixer")
 	void OnSetRowEditorVisibility(
-		UObject* InObject, bool bNewIsVisible, EObjectMixerTreeViewMode InViewMode = EObjectMixerTreeViewMode::FolderObjectSubObject) const override;
+		UObject* InObject, bool bNewIsVisible, EObjectMixerTreeViewMode InViewMode = EObjectMixerTreeViewMode::Folders) const override;
 
 	void OnSetRowEditorVisibility_Implementation(UObject* InObject, bool bNewIsVisible, EObjectMixerTreeViewMode InViewMode) const
 	{
diff --git a/Engine/Plugins/Editor/ObjectMixer/ObjectMixer/Source/ObjectMixer/Public/Views/List/ObjectMixerEditorListRow.h b/Engine/Plugins/Editor/ObjectMixer/ObjectMixer/Source/ObjectMixer/Public/Views/List/ObjectMixerEditorListRow.h
--- a/Engine/Plugins/Editor/ObjectMixer/ObjectMixer/Source/ObjectMixer/Public/Views/List/ObjectMixerEditorListRow.h
+++ b/Engine/Plugins/Editor/ObjectMixer/ObjectMixer/Source/ObjectMixer/Public/Views/List/ObjectMixerEditorListRow.h
@@ -38,11 +38,28 @@ struct OBJECTMIXEREDITOR_API FObjectMixerEditorListRow final : TSharedFromThis<F
 		None = 0,
 		Folder, // Usually an Outliner folder
 		ContainerObject, // Usually an actor that contains a matching subobject or is the attach parent of a matching actor
+		MatchingContainerObject, // An Actor that is a matching object and contains matching subobjects 
 		MatchingObject // The object that has the properties we wish to affect
 	};
 
 	~FObjectMixerEditorListRow();
 
+	bool operator==(const TSharedPtr<FObjectMixerEditorListRow>& Other) const
+	{
+		const UObject* ThisObject = GetObject();
+		const UObject* ThatObject = Other->GetObject();
+
+		const bool bAreObjectsEqual = ThisObject == ThatObject;
+
+		// If both are null, they're folders. Check display names.
+		if (bAreObjectsEqual && !ThisObject)
+		{
+			return GetDisplayNameOverride().EqualTo(Other->GetDisplayNameOverride());
+		}
+		
+		return bAreObjectsEqual;
+	}
+
 	void FlushReferences();
 	
 	FObjectMixerEditorListRow(
@@ -69,6 +86,9 @@ struct OBJECTMIXEREDITOR_API FObjectMixerEditorListRow final : TSharedFromThis<F
 	bool IsObjectRefInSelectedCollections() const;
 
 	[[nodiscard]] EObjectMixerEditorListRowType GetRowType() const;
+	void SetRowType(EObjectMixerEditorListRowType InNewRowType);
+
+	[[nodiscard]] bool IsHybridRow() const;
 
 	[[nodiscard]] int32 GetSortOrder() const;
 	void SetSortOrder(const int32 InNewOrder);
@@ -113,7 +133,7 @@ struct OBJECTMIXEREDITOR_API FObjectMixerEditorListRow final : TSharedFromThis<F
 
 	[[nodiscard]] FText GetDisplayName();
 
-	[[nodiscard]] const FText& GetDisplayNameOverride()
+	[[nodiscard]] const FText& GetDisplayNameOverride() const
 	{
 		return DisplayNameOverride;
 	}
diff --git a/Engine/Plugins/Editor/ObjectMixer/ObjectMixer/Source/ObjectMixer/Public/Views/List/SObjectMixerEditorListRow.h b/Engine/Plugins/Editor/ObjectMixer/ObjectMixer/Source/ObjectMixer/Public/Views/List/SObjectMixerEditorListRow.h
--- a/Engine/Plugins/Editor/ObjectMixer/ObjectMixer/Source/ObjectMixer/Public/Views/List/SObjectMixerEditorListRow.h
+++ b/Engine/Plugins/Editor/ObjectMixer/ObjectMixer/Source/ObjectMixer/Public/Views/List/SObjectMixerEditorListRow.h
@@ -62,6 +62,9 @@ private:
 
 	/** The offset applied to text widgets so that the text aligns with the column header text */
 	float TextBlockLeftPadding = 3.0f;
-	
+
+	/** Hybrid Rows are a combination of an actor and a single child subobject */
+	bool bIsHybridRow = false;
+
 	bool bIsHovered = false;
 };
diff --git a/Engine/Plugins/Editor/ObjectMixer/ObjectMixer/Source/ObjectMixer/Public/Views/MainPanel/ObjectMixerEditorMainPanel.h b/Engine/Plugins/Editor/ObjectMixer/ObjectMixer/Source/ObjectMixer/Public/Views/MainPanel/ObjectMixerEditorMainPanel.h
--- a/Engine/Plugins/Editor/ObjectMixer/ObjectMixer/Source/ObjectMixer/Public/Views/MainPanel/ObjectMixerEditorMainPanel.h
+++ b/Engine/Plugins/Editor/ObjectMixer/ObjectMixer/Source/ObjectMixer/Public/Views/MainPanel/ObjectMixerEditorMainPanel.h
@@ -184,7 +184,7 @@ private:
 	/**
 	 * Determines the style of the tree (flat list or hierarchy)
 	 */
-	EObjectMixerTreeViewMode TreeViewMode = EObjectMixerTreeViewMode::FolderObjectSubObject;
+	EObjectMixerTreeViewMode TreeViewMode = EObjectMixerTreeViewMode::Folders;
 
 	TWeakPtr<FObjectMixerEditorListRow> SoloRow = nullptr;
 
