Level Editor:

-Make LevelEditorContextMenu public

Object Mixer:

-Implement context menu for rows similar to the one in scene outliner

diff --git a/Engine/Plugins/Editor/ObjectMixer/ObjectMixer/Source/ObjectMixer/ObjectMixerEditor.build.cs b/Engine/Plugins/Editor/ObjectMixer/ObjectMixer/Source/ObjectMixer/ObjectMixerEditor.build.cs
--- a/Engine/Plugins/Editor/ObjectMixer/ObjectMixer/Source/ObjectMixer/ObjectMixerEditor.build.cs
+++ b/Engine/Plugins/Editor/ObjectMixer/ObjectMixer/Source/ObjectMixer/ObjectMixerEditor.build.cs
@@ -33,6 +33,7 @@ public class ObjectMixerEditor : ModuleRules
 				"EditorWidgets",
 				"InputCore",
 				"Kismet",
+				"LevelEditor",
 				"Projects",
 				"SceneOutliner",
 				"Slate",
@@ -40,6 +41,7 @@ public class ObjectMixerEditor : ModuleRules
 				"ToolMenus", 
 				"ToolWidgets",
 				"TypedElementFramework",
+				"TypedElementRuntime",
 				"UnrealEd",
 				"WorkspaceMenuStructure"
 			}
diff --git a/Engine/Plugins/Editor/ObjectMixer/ObjectMixer/Source/ObjectMixer/Private/Views/Widgets/ObjectMixerEditorListMenuContext.cpp b/Engine/Plugins/Editor/ObjectMixer/ObjectMixer/Source/ObjectMixer/Private/Views/Widgets/ObjectMixerEditorListMenuContext.cpp
--- a/Engine/Plugins/Editor/ObjectMixer/ObjectMixer/Source/ObjectMixer/Private/Views/Widgets/ObjectMixerEditorListMenuContext.cpp
+++ b/Engine/Plugins/Editor/ObjectMixer/ObjectMixer/Source/ObjectMixer/Private/Views/Widgets/ObjectMixerEditorListMenuContext.cpp
@@ -2,16 +2,20 @@
 
 #include "Views/Widgets/ObjectMixerEditorListMenuContext.h"
 
+#include "Elements/Framework/TypedElementSelectionSet.h"
+#include "Elements/Interfaces/TypedElementObjectInterface.h"
+#include "Framework/MultiBox/MultiBoxDefs.h"
 #include "LevelEditor.h"
+#include "LevelEditorContextMenu.h"
 #include "LevelEditorMenuContext.h"
-#include "Framework/MultiBox/MultiBoxDefs.h"
+#include "LevelEditorViewport.h"
+#include "Selection.h"
 #include "ToolMenu.h"
 #include "ToolMenus.h"
 #include "ToolMenuContext.h"
 #include "ToolMenuDelegates.h"
 #include "ToolMenuEntry.h"
 #include "ToolMenuSection.h"
-#include "UnrealEdGlobals.h"
 #include "Views/MainPanel/ObjectMixerEditorMainPanel.h"
 #include "Widgets/Layout/SBox.h"
 #include "Widgets/Input/SEditableTextBox.h"
@@ -32,41 +36,78 @@ TSharedPtr<SWidget> UObjectMixerEditorListMenuContext::CreateContextMenu(const F
 
 TSharedPtr<SWidget> UObjectMixerEditorListMenuContext::BuildContextMenu(const FObjectMixerEditorListMenuContextData& InData)
 {
-	RegisterContextMenu();
+	FLevelEditorContextMenu::RegisterComponentContextMenu();
+	FLevelEditorContextMenu::RegisterActorContextMenu();
+	FLevelEditorContextMenu::RegisterElementContextMenu();
+	RegisterObjectMixerContextMenu();
 
-	UObjectMixerEditorListMenuContext* ContextObject = NewObject<UObjectMixerEditorListMenuContext>();
-	ContextObject->Data = InData;
+	FLevelEditorModule& LevelEditorModule = FModuleManager::LoadModuleChecked<FLevelEditorModule>("LevelEditor");
+	TSharedPtr<ILevelEditor> LevelEditorPtr = LevelEditorModule.GetLevelEditorInstance().Pin();
+	check(LevelEditorPtr);
 
-	const FToolMenuContext Context(ContextObject);
+	TSharedPtr<FUICommandList> LevelEditorActionsList = LevelEditorPtr->GetLevelEditorActions();
+	FToolMenuContext Context;
+	Context.AppendCommandList(LevelEditorActionsList);
 
-	// todo: Uncomment all code when Level Editor Changes from 21454357 are pushed
-	// First get scene outliner-style context menu
-	// FLevelEditorModule& LevelEditorModule = FModuleManager::LoadModuleChecked<FLevelEditorModule>("LevelEditor");
-	// TSharedPtr<ILevelEditor> LevelEditorInstancePtr = LevelEditorModule.GetLevelEditorInstance().Pin();
-	//
-	// UToolMenu* BaseMenu = LevelEditorInstancePtr->GenerateLevelEditorContextMenu(
-	// 		ELevelEditorMenuContext::SceneOutliner, nullptr, FTypedElementHandle());
+	ULevelEditorContextMenuContext* LevelEditorContextObject = NewObject<ULevelEditorContextMenuContext>();
+	LevelEditorContextObject->LevelEditor = LevelEditorPtr;
+	LevelEditorContextObject->ContextType = ELevelEditorMenuContext::SceneOutliner;
+	LevelEditorContextObject->CurrentSelection = LevelEditorPtr->GetElementSelectionSet();
+	
+	for (FSelectedEditableComponentIterator It(GEditor->GetSelectedEditableComponentIterator()); It; ++It)
+	{
+		LevelEditorContextObject->SelectedComponents.Add(CastChecked<UActorComponent>(*It));
+	}
+	
+	Context.AddObject(LevelEditorContextObject, [](UObject* InContext)
+	{
+		ULevelEditorContextMenuContext* CastContext = CastChecked<ULevelEditorContextMenuContext>(InContext);
+		CastContext->CurrentSelection = nullptr;
+		CastContext->HitProxyElement.Release();
+	});
 
-	// Now generate custom sections to append
-	UToolMenus* ToolMenus = UToolMenus::Get();
-	UToolMenu* CustomMenu = ToolMenus->GenerateMenu(DefaultContextBaseMenuName, Context);
+	UObjectMixerEditorListMenuContext* ObjectMixerContextObject = NewObject<UObjectMixerEditorListMenuContext>();
+	ObjectMixerContextObject->Data = InData;
 
-	// for (const FToolMenuSection& Section : CustomMenu->Sections)
-	// {
-	// 	BaseMenu->Sections.Add(Section);
-	// }
+	Context.AddObject(ObjectMixerContextObject, [](UObject* InContext)
+	{
+		UObjectMixerEditorListMenuContext* CastContext = CastChecked<UObjectMixerEditorListMenuContext>(InContext);
+		CastContext->Data.SelectedItems.Empty();
+		CastContext->Data.MainPanelPtr.Reset();
+	});
 
-	//return ToolMenus->GenerateWidget(BaseMenu);
-	return ToolMenus->GenerateWidget(CustomMenu);
+	if (LevelEditorPtr->GetElementSelectionSet()->GetSelectedObjects<UActorComponent>().Num() == 0 &&
+		LevelEditorPtr->GetElementSelectionSet()->GetSelectedObjects<AActor>().Num() > 0)
+	{
+		const TArray<AActor*> SelectedActors = LevelEditorPtr->GetElementSelectionSet()->GetSelectedObjects<AActor>();
+	
+		// Get all menu extenders for this context menu from the level editor module
+		TArray<FLevelEditorModule::FLevelViewportMenuExtender_SelectedActors> MenuExtenderDelegates = LevelEditorModule.GetAllLevelViewportContextMenuExtenders();
+		TArray<TSharedPtr<FExtender>> Extenders;
+		for (int32 i = 0; i < MenuExtenderDelegates.Num(); ++i)
+		{
+			if (MenuExtenderDelegates[i].IsBound())
+			{
+				Extenders.Add(MenuExtenderDelegates[i].Execute(LevelEditorActionsList.ToSharedRef(), SelectedActors));
+			}
+		}
+	
+		if (Extenders.Num() > 0)
+		{
+			Context.AddExtender(FExtender::Combine(Extenders));
+		}
+	}
+
+	return UToolMenus::Get()->GenerateWidget("LevelEditor.ActorContextMenu", Context);
 }
 
-void UObjectMixerEditorListMenuContext::RegisterContextMenu()
+void UObjectMixerEditorListMenuContext::RegisterObjectMixerContextMenu()
 {
 	UToolMenus* ToolMenus = UToolMenus::Get();
 	
-	if (!ToolMenus->IsMenuRegistered(DefaultContextBaseMenuName))
+	if (ToolMenus->IsMenuRegistered("LevelEditor.ActorContextMenu"))
 	{
-		UToolMenu* Menu = ToolMenus->RegisterMenu(DefaultContextBaseMenuName);
+		UToolMenu* Menu = ToolMenus->ExtendMenu("LevelEditor.ActorContextMenu");
 	
 		Menu->AddDynamicSection("DynamicCollectionsSection", FNewToolMenuDelegate::CreateLambda(
 			[](UToolMenu* InMenu)
@@ -79,34 +120,36 @@ void UObjectMixerEditorListMenuContext::RegisterContextMenu()
 	
 				FToolMenuSection& Section = InMenu->FindOrAddSection("Collections");
 				Section.Label = LOCTEXT("CollectionsSectionName", "Collections");
+				
+				FObjectMixerEditorListMenuContextData& ContextData = Context->Data;
 			
 				Section.AddSubMenu(
 					"SelectCollectionsSubMenu",
 					LOCTEXT("SelectCollectionsSubmenu", "Select or Add Collection"),
 					LOCTEXT("SelectCollectionsSubmenu_Tooltip", "Select the collection to which you wish to assign this object."),
 					FNewToolMenuDelegate::CreateLambda(
-						[Context](UToolMenu* Menu)
+						[ContextData](UToolMenu* Menu)
 						{
 							FToolMenuEntry Args;
 							Args.Type = EMultiBlockType::Widget;
 							Args.MakeCustomWidget.BindLambda(
-								[Context](const FToolMenuContext&, const FToolMenuCustomWidgetContext&)
+								[ContextData](const FToolMenuContext&, const FToolMenuCustomWidgetContext&)
 								{
 									return SNew(SBox)
 											.MinDesiredWidth(200)
 											.Padding(8, 0)
 											[
-												SAssignNew(Context->EditableText, SEditableTextBox)
+												SNew(SEditableTextBox)
 												.HintText(LOCTEXT("NewCollectionEditableHintText", "Enter a new collection name..."))
-												.OnTextCommitted_Static(&UObjectMixerEditorListMenuContext::OnTextCommitted, Context)
-												.OnTextChanged_Static(&UObjectMixerEditorListMenuContext::OnTextChanged, Context)
+												.OnTextCommitted_Static(&UObjectMixerEditorListMenuContext::OnTextCommitted, ContextData)
+												.OnTextChanged_Static(&UObjectMixerEditorListMenuContext::OnTextChanged, ContextData)
 											]
 									;
 								}
 							);
 							Menu->AddMenuEntry("NewCollectionInput", Args);
 							
-							if (const TSharedPtr<FObjectMixerEditorMainPanel> MainPanel = Context->Data.MainPanelPtr.Pin())
+							if (const TSharedPtr<FObjectMixerEditorMainPanel> MainPanel = ContextData.MainPanelPtr.Pin())
 							{
 								if (TArray<FName> Collections = MainPanel->GetAllCollectionNames(); Collections.Num() > 0)
 								{
@@ -128,11 +171,11 @@ void UObjectMixerEditorListMenuContext::RegisterContextMenu()
 										   FText::Format(LOCTEXT("AddObjectsToCollectionTooltipFormat", "Add selected to collection '{0}'"), KeyText),
 										   FSlateIcon(),
 										   FUIAction(
-											   FExecuteAction::CreateStatic(&UObjectMixerEditorListMenuContext::OnClickCollectionMenuEntry, Key, Context),
+											   FExecuteAction::CreateStatic(&UObjectMixerEditorListMenuContext::OnClickCollectionMenuEntry, Key, ContextData),
 											   FCanExecuteAction(),
-											   FIsActionChecked::CreateStatic(&UObjectMixerEditorListMenuContext::AreAllObjectsInCollection, Key, Context)
+											   FIsActionChecked::CreateStatic(&UObjectMixerEditorListMenuContext::AreAllObjectsInCollection, Key, ContextData)
 										   ),
-										   EUserInterfaceActionType::Check
+										   EUserInterfaceActionType::ToggleButton
 									   );
 								   }
 								}
@@ -145,49 +188,43 @@ void UObjectMixerEditorListMenuContext::RegisterContextMenu()
 	}
 }
 
-void UObjectMixerEditorListMenuContext::OnTextCommitted(const FText& InText, ETextCommit::Type InCommitType, UObjectMixerEditorListMenuContext* Context)
+void UObjectMixerEditorListMenuContext::OnTextCommitted(const FText& InText, ETextCommit::Type InCommitType, const FObjectMixerEditorListMenuContextData ContextData)
 {
-	check(Context);
-	
 	if (InCommitType == ETextCommit::OnEnter)
 	{
 		UE_LOG(LogTemp, Warning, TEXT("%hs: %s"), __FUNCTION__, *InText.ToString());
 
-		AddObjectsToCollection(*InText.ToString(), Context);
+		AddObjectsToCollection(*InText.ToString(), ContextData);
 	}
 }
 
-void UObjectMixerEditorListMenuContext::OnTextChanged(const FText& InText, UObjectMixerEditorListMenuContext* Context)
+void UObjectMixerEditorListMenuContext::OnTextChanged(const FText& InText, const FObjectMixerEditorListMenuContextData ContextData)
 {
-	check(Context);
-	
 	if (const FString AsString = InText.ToString(); AsString.Len() > 12)
 	{
-		Context->EditableText->SetText(FText::FromString(AsString.LeftChop(1)));
+		//ContextData.EditableText->SetText(FText::FromString(AsString.LeftChop(1)));
 	}
 }
 
-void UObjectMixerEditorListMenuContext::OnClickCollectionMenuEntry(const FName Key, UObjectMixerEditorListMenuContext* Context)
+void UObjectMixerEditorListMenuContext::OnClickCollectionMenuEntry(const FName Key, const FObjectMixerEditorListMenuContextData ContextData)
 {
-	if (AreAllObjectsInCollection(Key, Context))
+	if (AreAllObjectsInCollection(Key, ContextData))
 	{
-		RemoveObjectsFromCollection(Key, Context);
+		RemoveObjectsFromCollection(Key, ContextData);
 	}
 	else
 	{
-		AddObjectsToCollection(Key, Context);
+		AddObjectsToCollection(Key, ContextData);
 	}
 }
 
-void UObjectMixerEditorListMenuContext::AddObjectsToCollection(const FName Key, UObjectMixerEditorListMenuContext* Context)
+void UObjectMixerEditorListMenuContext::AddObjectsToCollection(const FName Key, const FObjectMixerEditorListMenuContextData ContextData)
 {
-	check(Context);
-
-	if (const TSharedPtr<FObjectMixerEditorMainPanel> MainPanel = Context->Data.MainPanelPtr.Pin())
+	if (const TSharedPtr<FObjectMixerEditorMainPanel> MainPanel = ContextData.MainPanelPtr.Pin())
 	{
 		TSet<FSoftObjectPath> ObjectPaths;
 
-		for (const TSharedPtr<FObjectMixerEditorListRow>& Item : Context->Data.SelectedItems)
+		for (const TSharedPtr<FObjectMixerEditorListRow>& Item : ContextData.SelectedItems)
 		{
 			if (const UObject* Object = Item->GetObject())
 			{
@@ -199,15 +236,13 @@ void UObjectMixerEditorListMenuContext::AddObjectsToCollection(const FName Key,
 	}
 }
 
-void UObjectMixerEditorListMenuContext::RemoveObjectsFromCollection(const FName Key, UObjectMixerEditorListMenuContext* Context)
+void UObjectMixerEditorListMenuContext::RemoveObjectsFromCollection(const FName Key, const FObjectMixerEditorListMenuContextData ContextData)
 {
-	check(Context);
-
-	if (const TSharedPtr<FObjectMixerEditorMainPanel> MainPanel = Context->Data.MainPanelPtr.Pin())
+	if (const TSharedPtr<FObjectMixerEditorMainPanel> MainPanel = ContextData.MainPanelPtr.Pin())
 	{
 		TSet<FSoftObjectPath> ObjectPaths;
 
-		for (const TSharedPtr<FObjectMixerEditorListRow>& Item : Context->Data.SelectedItems)
+		for (const TSharedPtr<FObjectMixerEditorListRow>& Item : ContextData.SelectedItems)
 		{
 			if (const UObject* Object = Item->GetObject())
 			{
@@ -219,15 +254,13 @@ void UObjectMixerEditorListMenuContext::RemoveObjectsFromCollection(const FName
 	}
 }
 
-bool UObjectMixerEditorListMenuContext::AreAllObjectsInCollection(const FName Key, UObjectMixerEditorListMenuContext* Context)
+bool UObjectMixerEditorListMenuContext::AreAllObjectsInCollection(const FName Key, const FObjectMixerEditorListMenuContextData ContextData)
 {
-	check(Context);
-
 	bool bAreAllSelectedObjectsInCollection = false;
 
-	if (const TSharedPtr<FObjectMixerEditorMainPanel> MainPanel = Context->Data.MainPanelPtr.Pin())
+	if (const TSharedPtr<FObjectMixerEditorMainPanel> MainPanel = ContextData.MainPanelPtr.Pin())
 	{
-		for (const TSharedPtr<FObjectMixerEditorListRow>& Item : Context->Data.SelectedItems)
+		for (const TSharedPtr<FObjectMixerEditorListRow>& Item : ContextData.SelectedItems)
 		{
 			if (const UObject* Object = Item->GetObject())
 			{
diff --git a/Engine/Plugins/Editor/ObjectMixer/ObjectMixer/Source/ObjectMixer/Public/Views/Widgets/ObjectMixerEditorListMenuContext.h b/Engine/Plugins/Editor/ObjectMixer/ObjectMixer/Source/ObjectMixer/Public/Views/Widgets/ObjectMixerEditorListMenuContext.h
--- a/Engine/Plugins/Editor/ObjectMixer/ObjectMixer/Source/ObjectMixer/Public/Views/Widgets/ObjectMixerEditorListMenuContext.h
+++ b/Engine/Plugins/Editor/ObjectMixer/ObjectMixer/Source/ObjectMixer/Public/Views/Widgets/ObjectMixerEditorListMenuContext.h
@@ -24,7 +24,7 @@ public:
 
 	static TSharedPtr<SWidget> CreateContextMenu(const FObjectMixerEditorListMenuContextData InData);
 	static TSharedPtr<SWidget> BuildContextMenu(const FObjectMixerEditorListMenuContextData& InData);
-	static void RegisterContextMenu();
+	static void RegisterObjectMixerContextMenu();
 
 	FObjectMixerEditorListMenuContextData Data;
 	
@@ -32,13 +32,11 @@ public:
 
 private:
 
-	static void OnTextCommitted(const FText& InText, ETextCommit::Type InCommitType, UObjectMixerEditorListMenuContext* Context);
-	static void OnTextChanged(const FText& InText, UObjectMixerEditorListMenuContext* Context);
+	static void OnTextCommitted(const FText& InText, ETextCommit::Type InCommitType, const FObjectMixerEditorListMenuContextData ContextData);
+	static void OnTextChanged(const FText& InText, const FObjectMixerEditorListMenuContextData ContextData);
 	
-	static void OnClickCollectionMenuEntry(const FName Key, UObjectMixerEditorListMenuContext* Context);
-	static void AddObjectsToCollection(const FName Key, UObjectMixerEditorListMenuContext* Context);
-	static void RemoveObjectsFromCollection(const FName Key, UObjectMixerEditorListMenuContext* Context);
-	static bool AreAllObjectsInCollection(const FName Key, UObjectMixerEditorListMenuContext* Context);
-
-	TSharedPtr<class SEditableTextBox> EditableText;
+	static void OnClickCollectionMenuEntry(const FName Key, const FObjectMixerEditorListMenuContextData ContextData);
+	static void AddObjectsToCollection(const FName Key, const FObjectMixerEditorListMenuContextData ContextData);
+	static void RemoveObjectsFromCollection(const FName Key, const FObjectMixerEditorListMenuContextData ContextData);
+	static bool AreAllObjectsInCollection(const FName Key, const FObjectMixerEditorListMenuContextData ContextData);
 };
diff --git a/Engine/Source/Editor/LevelEditor/Private/LevelEditorContextMenu.h b/Engine/Source/Editor/LevelEditor/Private/LevelEditorContextMenu.h
deleted file mode 100644
diff --git a/Engine/Source/Editor/LevelEditor/Public/LevelEditorContextMenu.h b/Engine/Source/Editor/LevelEditor/Public/LevelEditorContextMenu.h
new file mode 100644
--- /dev/null
+++ b/Engine/Source/Editor/LevelEditor/Public/LevelEditorContextMenu.h
@@ -0,0 +1,89 @@
+// Copyright Epic Games, Inc. All Rights Reserved.
+
+#pragma once
+
+#include "CoreMinimal.h"
+#include "Editor/UnrealEdTypes.h"
+#include "LevelEditorMenuContext.h"
+
+class FExtender;
+class UToolMenu;
+class UTypedElementSelectionSet;
+struct FToolMenuContext;
+class SLevelEditor;
+class SWidget;
+struct FToolMenuSection;
+
+/**
+ * Context menu construction class 
+ */
+class LEVELEDITOR_API FLevelEditorContextMenu
+{
+
+public:
+
+	/**
+	 * Summons the level viewport context menu
+	 * @param	LevelEditor		The level editor using this menu.
+	 * @param	ContextType		The context we should use to specialize this menu
+	 * @param	HitProxyActor	The hitproxy actor in the case the ContextType is Viewport
+	 */
+	static void SummonMenu( const TSharedRef< class SLevelEditor >& LevelEditor, ELevelEditorMenuContext ContextType, const FTypedElementHandle& HitProxyElement = FTypedElementHandle());
+
+	/**
+	 * Summons the viewport view option menu
+	 * @param LevelEditor		The level editor using this menu.
+	 */
+	static void SummonViewOptionMenu( const TSharedRef< class SLevelEditor >& LevelEditor, const ELevelViewportType ViewOption );
+
+	/**
+	 * Creates a widget for the context menu that can be inserted into a pop-up window
+	 *
+	 * @param	LevelEditor		The level editor using this menu.
+	 * @param	ContextType		The context we should use to specialize this menu
+	 * @param	Extender		Allows extension of this menu based on context.
+	 * @param	HitProxyActor	The hitproxy actor in the case the ContextType is Viewport
+	 * @return	Widget for this context menu
+	 */
+	static TSharedRef< SWidget > BuildMenuWidget(TWeakPtr< ILevelEditor > LevelEditor, ELevelEditorMenuContext ContextType, TSharedPtr<FExtender> Extender = TSharedPtr<FExtender>(), const FTypedElementHandle& HitProxyElement = FTypedElementHandle());
+
+	/**
+	 * Populates the specified menu builder for the context menu that can be inserted into a pop-up window
+	 *
+	 * @param	Menu			The menu to fill
+	 * @param	LevelEditor		The level editor using this menu.
+	 * @param	ContextType		The context we should use to specialize this menu
+	 * @param	Extender		Allows extension of this menu based on context.
+	 * @param	HitProxyActor	The hitproxy actor in the case the ContextType is Viewport
+	 */
+	static UToolMenu* GenerateMenu(TWeakPtr< ILevelEditor > LevelEditor, ELevelEditorMenuContext ContextType, TSharedPtr<FExtender> Extender = TSharedPtr<FExtender>(), const FTypedElementHandle& HitProxyElement = FTypedElementHandle());
+
+	/* Adds required information to Context for build menu based on current selection */
+	static FName InitMenuContext(FToolMenuContext& Context, TWeakPtr<ILevelEditor> LevelEditor, ELevelEditorMenuContext ContextType, const FTypedElementHandle& HitProxyElement = FTypedElementHandle());
+
+	/* Returns name of menu to display based on current selection */
+	static FName GetContextMenuName(ELevelEditorMenuContext ContextType, const UTypedElementSelectionSet* InSelectionSet);
+
+	/* Returns a user-readable title for the menu to display; the title can be displayed in UI like the menu bar */
+	static FText GetContextMenuTitle(ELevelEditorMenuContext ContextType, const UTypedElementSelectionSet* InSelectionSet);
+
+	/* Returns a user-readable tooltip describing the menu to display */
+	static FText GetContextMenuToolTip(ELevelEditorMenuContext ContextType, const UTypedElementSelectionSet* InSelectionSet);
+
+	static void RegisterComponentContextMenu();
+	static void RegisterActorContextMenu();
+	static void RegisterElementContextMenu();
+	static void RegisterSceneOutlinerContextMenu();
+	static void RegisterMenuBarEmptyContextMenu();
+	static void RegisterEmptySelectionContextMenu();
+
+	static void AddPlayFromHereSubMenu(FToolMenuSection& Section);
+
+	/**
+	 * Builds the actor group menu
+	 *
+	 * @param Menu		The menu to add items to.
+	 * @param SelectedActorInfo	Information about the selected actors.
+	 */
+	static void BuildGroupMenu(FToolMenuSection& Section, const struct FSelectedActorInfo& SelectedActorInfo);
+};
