Object Mixer:

-Dimly highlight parent rows when child row is selected

diff --git a/Engine/Plugins/Editor/ObjectMixer/ObjectMixer/Source/ObjectMixer/Private/Views/List/SObjectMixerEditorList.cpp b/Engine/Plugins/Editor/ObjectMixer/ObjectMixer/Source/ObjectMixer/Private/Views/List/SObjectMixerEditorList.cpp
--- a/Engine/Plugins/Editor/ObjectMixer/ObjectMixer/Source/ObjectMixer/Private/Views/List/SObjectMixerEditorList.cpp
+++ b/Engine/Plugins/Editor/ObjectMixer/ObjectMixer/Source/ObjectMixer/Private/Views/List/SObjectMixerEditorList.cpp
@@ -70,6 +70,7 @@ void SObjectMixerEditorList::Construct(const FArguments& InArgs, TSharedRef<FObj
 				SAssignNew(TreeViewPtr, STreeView<FObjectMixerEditorListRowPtr>)
 				.HeaderRow(HeaderRow)
 				.SelectionMode(ESelectionMode::Multi)
+				.HighlightParentNodesForSelection(true)
 				.OnContextMenuOpening_Lambda([this]()
 				{
 					TSharedPtr<SWidget> Widget;
