Object Mixer:

-Fix for: Not showing Hierarchies or Folders

diff --git a/Engine/Plugins/Editor/ObjectMixer/ObjectMixer/Source/ObjectMixer/Private/Views/List/ObjectMixerEditorListRow.cpp b/Engine/Plugins/Editor/ObjectMixer/ObjectMixer/Source/ObjectMixer/Private/Views/List/ObjectMixerEditorListRow.cpp
--- a/Engine/Plugins/Editor/ObjectMixer/ObjectMixer/Source/ObjectMixer/Private/Views/List/ObjectMixerEditorListRow.cpp
+++ b/Engine/Plugins/Editor/ObjectMixer/ObjectMixer/Source/ObjectMixer/Private/Views/List/ObjectMixerEditorListRow.cpp
@@ -97,10 +97,42 @@ void FObjectMixerEditorListRow::SetRowType(EObjectMixerEditorListRowType InNewRo
 	RowType = InNewRowType;
 }
 
-bool FObjectMixerEditorListRow::IsHybridRow() const
+int32 FObjectMixerEditorListRow::GetHybridRowIndex() const
 {
-	return GetRowType() == FObjectMixerEditorListRow::ContainerObject &&
-			GetChildCount() == 1 && GetChildRows()[0]->GetChildCount() == 0;
+	if (GetRowType() == FObjectMixerEditorListRow::ContainerObject)
+	{
+		if (const UObject* ThisObject = GetObject())
+		{
+			TArray<int32> HybridCandidates;
+		
+			for (int32 ChildItr = 0; ChildItr < GetChildCount(); ChildItr++)
+			{
+				const FObjectMixerEditorListRowPtr& ChildRow = GetChildRows()[ChildItr];
+			
+				if (ChildRow->GetRowType() == FObjectMixerEditorListRow::MatchingObject)
+				{
+					if (const UObject* ChildObject = ChildRow->GetObject(); ChildObject->GetOuter() == ThisObject)
+					{
+						HybridCandidates.Add(ChildItr);
+
+						if (HybridCandidates.Num() > 1)
+						{
+							// There can only be one row to hybrid with.
+							// If there's more than one candidate, don't hybrid.
+							break;
+						}
+					}
+				}
+			}
+			
+			if (HybridCandidates.Num() == 1)
+			{
+				return HybridCandidates[0];
+			}
+		}
+	}
+	
+	return INDEX_NONE;
 }
 
 int32 FObjectMixerEditorListRow::GetSortOrder() const
diff --git a/Engine/Plugins/Editor/ObjectMixer/ObjectMixer/Source/ObjectMixer/Private/Views/List/SObjectMixerEditorList.cpp b/Engine/Plugins/Editor/ObjectMixer/ObjectMixer/Source/ObjectMixer/Private/Views/List/SObjectMixerEditorList.cpp
--- a/Engine/Plugins/Editor/ObjectMixer/ObjectMixer/Source/ObjectMixer/Private/Views/List/SObjectMixerEditorList.cpp
+++ b/Engine/Plugins/Editor/ObjectMixer/ObjectMixer/Source/ObjectMixer/Private/Views/List/SObjectMixerEditorList.cpp
@@ -997,14 +997,16 @@ void CreateOuterRowsForTopLevelRow(
 		}
 		TopLevelRow = OuterRow;
 			
-		UObject* Outer = InObject->GetOuter();
+		UObject* Outer = nullptr;
+		
+		if (const AActor* AsActor = Cast<AActor>(InObject))
+		{
+			Outer = AsActor->GetAttachParentActor();
+		}
 
 		if (!Outer)
 		{
-			if (const AActor* AsActor = Cast<AActor>(InObject))
-			{
-				Outer = AsActor->GetAttachParentActor();
-			}
+			Outer = InObject->GetOuter();
 		}
 
 		// Recurse with Outer object
@@ -1180,15 +1182,17 @@ void SObjectMixerEditorList::OnGetRowChildren(FObjectMixerEditorListRowPtr Row,
 {
 	if (Row.IsValid())
 	{
-		if (!Row->IsHybridRow())
+		OutChildren = Row->GetChildRows();
+		
+		if (const int32 HybridIndex = Row->GetHybridRowIndex(); HybridIndex != INDEX_NONE)
 		{
-			OutChildren = Row->GetChildRows();
+			OutChildren.RemoveAt(HybridIndex);
+		}
 
-			if (Row->GetShouldExpandAllChildren())
-			{
-				SetChildExpansionRecursively(Row, true);
-				Row->SetShouldExpandAllChildren(false);
-			}
+		if (Row->GetShouldExpandAllChildren())
+		{
+			SetChildExpansionRecursively(Row, true);
+			Row->SetShouldExpandAllChildren(false);
 		}
 	}
 }
diff --git a/Engine/Plugins/Editor/ObjectMixer/ObjectMixer/Source/ObjectMixer/Private/Views/List/SObjectMixerEditorListRow.cpp b/Engine/Plugins/Editor/ObjectMixer/ObjectMixer/Source/ObjectMixer/Private/Views/List/SObjectMixerEditorListRow.cpp
--- a/Engine/Plugins/Editor/ObjectMixer/ObjectMixer/Source/ObjectMixer/Private/Views/List/SObjectMixerEditorListRow.cpp
+++ b/Engine/Plugins/Editor/ObjectMixer/ObjectMixer/Source/ObjectMixer/Private/Views/List/SObjectMixerEditorListRow.cpp
@@ -37,7 +37,7 @@ void SObjectMixerEditorListRow::Construct(
 	check(InRow.IsValid());
 
 	Item = InRow;
-	bIsHybridRow = Item.Pin()->IsHybridRow();
+	HybridRowIndex = Item.Pin()->GetHybridRowIndex();
 
 	SMultiColumnTableRow<FObjectMixerEditorListRowPtr>::Construct(
 		FSuperRowType::FArguments()
@@ -63,7 +63,7 @@ void SObjectMixerEditorListRow::Construct(
 TSharedRef<SWidget> SObjectMixerEditorListRow::GenerateWidgetForColumn(const FName& InColumnName)
 {
 	const FObjectMixerEditorListRowPtr PinnedItem =
-		bIsHybridRow ? Item.Pin()->GetChildRows()[0] : Item.Pin();
+		HybridRowIndex != INDEX_NONE ? Item.Pin()->GetChildRows()[HybridRowIndex] : Item.Pin();
 
 	if (const TSharedPtr<SWidget> CellWidget = GenerateCells(InColumnName, PinnedItem))
 	{
diff --git a/Engine/Plugins/Editor/ObjectMixer/ObjectMixer/Source/ObjectMixer/Public/Views/List/ObjectMixerEditorListRow.h b/Engine/Plugins/Editor/ObjectMixer/ObjectMixer/Source/ObjectMixer/Public/Views/List/ObjectMixerEditorListRow.h
--- a/Engine/Plugins/Editor/ObjectMixer/ObjectMixer/Source/ObjectMixer/Public/Views/List/ObjectMixerEditorListRow.h
+++ b/Engine/Plugins/Editor/ObjectMixer/ObjectMixer/Source/ObjectMixer/Public/Views/List/ObjectMixerEditorListRow.h
@@ -88,7 +88,8 @@ struct OBJECTMIXEREDITOR_API FObjectMixerEditorListRow final : TSharedFromThis<F
 	[[nodiscard]] EObjectMixerEditorListRowType GetRowType() const;
 	void SetRowType(EObjectMixerEditorListRowType InNewRowType);
 
-	[[nodiscard]] bool IsHybridRow() const;
+	/** If this is a hybrid row, return the index of the child used to hybrid with. */
+	[[nodiscard]] int32 GetHybridRowIndex() const;
 
 	[[nodiscard]] int32 GetSortOrder() const;
 	void SetSortOrder(const int32 InNewOrder);
diff --git a/Engine/Plugins/Editor/ObjectMixer/ObjectMixer/Source/ObjectMixer/Public/Views/List/SObjectMixerEditorListRow.h b/Engine/Plugins/Editor/ObjectMixer/ObjectMixer/Source/ObjectMixer/Public/Views/List/SObjectMixerEditorListRow.h
--- a/Engine/Plugins/Editor/ObjectMixer/ObjectMixer/Source/ObjectMixer/Public/Views/List/SObjectMixerEditorListRow.h
+++ b/Engine/Plugins/Editor/ObjectMixer/ObjectMixer/Source/ObjectMixer/Public/Views/List/SObjectMixerEditorListRow.h
@@ -64,7 +64,7 @@ private:
 	float TextBlockLeftPadding = 3.0f;
 
 	/** Hybrid Rows are a combination of an actor and a single child subobject */
-	bool bIsHybridRow = false;
+	int32 HybridRowIndex = INDEX_NONE;
 
 	bool bIsHovered = false;
 };
