Object Mixer:

-Fix for: Mute / Unmute and Solo functionality broken

diff --git a/Engine/Plugins/Editor/ObjectMixer/ObjectMixer/Source/ObjectMixer/Private/Views/List/SObjectMixerEditorListRow.cpp b/Engine/Plugins/Editor/ObjectMixer/ObjectMixer/Source/ObjectMixer/Private/Views/List/SObjectMixerEditorListRow.cpp
--- a/Engine/Plugins/Editor/ObjectMixer/ObjectMixer/Source/ObjectMixer/Private/Views/List/SObjectMixerEditorListRow.cpp
+++ b/Engine/Plugins/Editor/ObjectMixer/ObjectMixer/Source/ObjectMixer/Private/Views/List/SObjectMixerEditorListRow.cpp
@@ -302,6 +302,9 @@ TSharedPtr<SWidget> SObjectMixerEditorListRow::GenerateCells(
 {
 	check(PinnedItem.IsValid());
 	
+	const bool bIsRowTypeNone = PinnedItem->GetRowType() == FObjectMixerEditorListRow::None;
+	const bool bIsHybridRow = HybridRowIndex != INDEX_NONE;
+	
 	if (InColumnName.IsEqual(SObjectMixerEditorList::ItemNameColumnName))
 	{
 		TSharedRef<SHorizontalBox> HBox = SNew(SHorizontalBox);
@@ -367,12 +370,7 @@ TSharedPtr<SWidget> SObjectMixerEditorListRow::GenerateCells(
 
 	if (InColumnName.IsEqual(SObjectMixerEditorList::EditorVisibilityColumnName))
 	{
-		if (PinnedItem->GetRowType() == FObjectMixerEditorListRow::None)
-		{
-			return nullptr;
-		}
-
-		if (PinnedItem->GetObject() && !PinnedItem->GetObject()->IsA(AActor::StaticClass()))
+		if (!bIsHybridRow && PinnedItem->GetObject() && !PinnedItem->GetObject()->IsA(AActor::StaticClass()))
 		{
 			return nullptr;
 		}
@@ -415,12 +413,7 @@ TSharedPtr<SWidget> SObjectMixerEditorListRow::GenerateCells(
 
 	if (InColumnName.IsEqual(SObjectMixerEditorList::EditorVisibilitySoloColumnName))
 	{
-		if (PinnedItem->GetRowType() == FObjectMixerEditorListRow::None)
-		{
-			return nullptr;
-		}
-
-		if (PinnedItem->GetObject() && !PinnedItem->GetObject()->IsA(AActor::StaticClass()))
+		if (!bIsHybridRow && PinnedItem->GetObject() && !PinnedItem->GetObject()->IsA(AActor::StaticClass()))
 		{
 			return nullptr;
 		}
