Object Mixer:

-Delete Light leaves empty row

diff --git a/Engine/Plugins/Editor/ObjectMixer/ObjectMixer/Source/ObjectMixer/Private/Views/List/SObjectMixerEditorList.cpp b/Engine/Plugins/Editor/ObjectMixer/ObjectMixer/Source/ObjectMixer/Private/Views/List/SObjectMixerEditorList.cpp
--- a/Engine/Plugins/Editor/ObjectMixer/ObjectMixer/Source/ObjectMixer/Private/Views/List/SObjectMixerEditorList.cpp
+++ b/Engine/Plugins/Editor/ObjectMixer/ObjectMixer/Source/ObjectMixer/Private/Views/List/SObjectMixerEditorList.cpp
@@ -937,13 +937,16 @@ void SObjectMixerEditorList::BuildPerformanceCacheAndGenerateHeaderIfNeeded()
 	}
 }
 
-bool DoesWorldObjectHaveAcceptableClass(const UObject* Object, const TSet<UClass*>& ObjectClassesToFilterCache)
+bool DoesValidWorldObjectHaveAcceptableClass(const UObject* Object, const TSet<UClass*>& ObjectClassesToFilterCache)
 {
-	for (UClass* Class : ObjectClassesToFilterCache)
+	if (IsValid(Object))
 	{
-		if (Object->IsA(Class))
+		for (UClass* Class : ObjectClassesToFilterCache)
 		{
-			return true;
+			if (Object->IsA(Class))
+			{
+				return true;
+			}
 		}
 	}
 
@@ -1081,7 +1084,7 @@ void SObjectMixerEditorList::GenerateTreeView()
 	ForEachObjectWithOuter(EditorWorld,
 		[this, &AllMatchingObjects](UObject* Object)
 	{
-		if (DoesWorldObjectHaveAcceptableClass(Object, ObjectClassesToFilterCache))
+		if (DoesValidWorldObjectHaveAcceptableClass(Object, ObjectClassesToFilterCache))
 		{
 			AllMatchingObjects.Add(Object);
 		}
