Object Mixer:

-Visual adjustments to light mixer panel

diff --git a/Engine/Plugins/Editor/ObjectMixer/LightMixer/Source/LightMixer/Private/LightMixerStyle.cpp b/Engine/Plugins/Editor/ObjectMixer/LightMixer/Source/LightMixer/Private/LightMixerStyle.cpp
--- a/Engine/Plugins/Editor/ObjectMixer/LightMixer/Source/LightMixer/Private/LightMixerStyle.cpp
+++ b/Engine/Plugins/Editor/ObjectMixer/LightMixer/Source/LightMixer/Private/LightMixerStyle.cpp
@@ -72,7 +72,11 @@ TSharedRef< FSlateStyleSet > FLightMixerStyle::Create()
 {
 	TSharedRef<FSlateStyleSet> Style = MakeShared<FSlateStyleSet>("LightMixer");
 	
-	//Style->SetContentRoot(FPaths::Combine(Plugin->GetBaseDir(), TEXT("Resources")));
+	TSharedPtr<IPlugin> Plugin = IPluginManager::Get().FindPlugin(TEXT("LightMixer"));
+	if (ensure(Plugin.IsValid()))
+	{
+		Style->SetContentRoot(FPaths::Combine(Plugin->GetBaseDir(), TEXT("Resources")));
+	}
 
 	Style->SetCoreContentRoot(FPaths::EngineContentDir() / TEXT("Slate"));
 
diff --git a/Engine/Plugins/Editor/ObjectMixer/ObjectMixer/Source/ObjectMixer/Private/ObjectMixerEditorStyle.cpp b/Engine/Plugins/Editor/ObjectMixer/ObjectMixer/Source/ObjectMixer/Private/ObjectMixerEditorStyle.cpp
--- a/Engine/Plugins/Editor/ObjectMixer/ObjectMixer/Source/ObjectMixer/Private/ObjectMixerEditorStyle.cpp
+++ b/Engine/Plugins/Editor/ObjectMixer/ObjectMixer/Source/ObjectMixer/Private/ObjectMixerEditorStyle.cpp
@@ -72,20 +72,19 @@ TSharedRef< FSlateStyleSet > FObjectMixerEditorStyle::Create()
 {
 	TSharedRef<FSlateStyleSet> Style = MakeShared<FSlateStyleSet>("ObjectMixerEditor");
 	
-	//Style->SetContentRoot(FPaths::Combine(Plugin->GetBaseDir(), TEXT("Resources")));
+	TSharedPtr<IPlugin> Plugin = IPluginManager::Get().FindPlugin(TEXT("ObjectMixer"));
+	if (ensure(Plugin.IsValid()))
+	{
+		Style->SetContentRoot(FPaths::Combine(Plugin->GetBaseDir(), TEXT("Resources")));
+	}
 
-	//Style->SetCoreContentRoot(FPaths::EngineContentDir() / TEXT("Slate"));
+	Style->SetCoreContentRoot(FPaths::EngineContentDir() / TEXT("Slate"));
 
 	// Icons
 	Style->Set("ObjectMixer.ToolbarButton", new IMAGE_BRUSH_SVG("Icons/ObjectMixer", Icon40x40));
 	Style->Set("ObjectMixer.ToolbarButton.Small", new IMAGE_BRUSH_SVG("Icons/ObjectMixer", Icon20x20));
-
-	// Brush
-	Style->Set("ObjectMixerEditor.BrightBorder", new FSlateColorBrush(FStyleColors::Header));
-
-	// Border colors for Results view
-	Style->Set("ObjectMixerEditor.HeaderRowBorder", new FSlateColorBrush(FStyleColors::Black));
-	Style->Set("ObjectMixerEditor.DefaultBorder", new FSlateColorBrush(FStyleColors::Transparent));
+	
+	Style->Set("ObjectMixer.Solo", new IMAGE_BRUSH_SVG("Icons/Solo", Icon16x16));
 	
 	return Style;
 }
diff --git a/Engine/Plugins/Editor/ObjectMixer/ObjectMixer/Source/ObjectMixer/Private/Views/List/SObjectMixerEditorList.cpp b/Engine/Plugins/Editor/ObjectMixer/ObjectMixer/Source/ObjectMixer/Private/Views/List/SObjectMixerEditorList.cpp
--- a/Engine/Plugins/Editor/ObjectMixer/ObjectMixer/Source/ObjectMixer/Private/Views/List/SObjectMixerEditorList.cpp
+++ b/Engine/Plugins/Editor/ObjectMixer/ObjectMixer/Source/ObjectMixer/Private/Views/List/SObjectMixerEditorList.cpp
@@ -723,7 +723,7 @@ TSharedPtr<SHeaderRow> SObjectMixerEditorList::GenerateHeaderRow()
 					[
 						SNew(SImage)
 						.ColorAndOpacity(FSlateColor::UseForeground())
-						.Image(FAppStyle::Get().GetBrush("MediaAsset.AssetActions.Solo.Small"))
+						.Image(FObjectMixerEditorStyle::Get().GetBrush("ObjectMixer.Solo"))
 						.ToolTipText(LOCTEXT("SoloColumnTooltip", "Solo"))
 					]
 				];
diff --git a/Engine/Plugins/Editor/ObjectMixer/ObjectMixer/Source/ObjectMixer/Private/Views/List/SObjectMixerEditorListRow.cpp b/Engine/Plugins/Editor/ObjectMixer/ObjectMixer/Source/ObjectMixer/Private/Views/List/SObjectMixerEditorListRow.cpp
--- a/Engine/Plugins/Editor/ObjectMixer/ObjectMixer/Source/ObjectMixer/Private/Views/List/SObjectMixerEditorListRow.cpp
+++ b/Engine/Plugins/Editor/ObjectMixer/ObjectMixer/Source/ObjectMixer/Private/Views/List/SObjectMixerEditorListRow.cpp
@@ -90,10 +90,9 @@ TSharedRef<SWidget> SObjectMixerEditorListRow::GenerateWidgetForColumn(const FNa
 				];
 		}
 		
-		return SNew(SBorder)
+		return SNew(SBox)
 				   .HAlign(HAlign_Fill)
 				   .VAlign(VAlign_Center)
-				   .BorderImage(GetBorderImage(PinnedItem->GetRowType()))
 			   [
 				   CellWidget.ToSharedRef()
 			   ];
@@ -291,12 +290,6 @@ const FSlateBrush* SObjectMixerEditorListRow::GetVisibilityBrush() const
 	return bIsHovered ? NotVisibleHoveredBrush : NotVisibleNotHoveredBrush;
 }
 
-const FSlateBrush* SObjectMixerEditorListRow::GetBorderImage(
-	const FObjectMixerEditorListRow::EObjectMixerEditorListRowType InRowType)
-{
-	return FObjectMixerEditorStyle::Get().GetBrush("ObjectMixerEditor.DefaultBorder");
-}
-
 TSharedPtr<SWidget> SObjectMixerEditorListRow::GenerateCells(
 	const FName& InColumnName, const TSharedPtr<FObjectMixerEditorListRow> PinnedItem)
 {
diff --git a/Engine/Plugins/Editor/ObjectMixer/ObjectMixer/Source/ObjectMixer/Private/Views/MainPanel/SObjectMixerEditorMainPanel.cpp b/Engine/Plugins/Editor/ObjectMixer/ObjectMixer/Source/ObjectMixer/Private/Views/MainPanel/SObjectMixerEditorMainPanel.cpp
--- a/Engine/Plugins/Editor/ObjectMixer/ObjectMixer/Source/ObjectMixer/Private/Views/MainPanel/SObjectMixerEditorMainPanel.cpp
+++ b/Engine/Plugins/Editor/ObjectMixer/ObjectMixer/Source/ObjectMixer/Private/Views/MainPanel/SObjectMixerEditorMainPanel.cpp
@@ -77,7 +77,7 @@ TSharedRef<SWidget> SObjectMixerEditorMainPanel::GenerateToolbar()
    .HAlign(HAlign_Left)
    .VAlign(VAlign_Center)
    .AutoWidth()
-   .Padding(FMargin(0, 4))
+   .Padding(FMargin(8, 4))
    [
 	   SNew(SPositiveActionButton)
 	   .Text(LOCTEXT("AddObject", "Add"))
@@ -106,7 +106,6 @@ TSharedRef<SWidget> SObjectMixerEditorMainPanel::GenerateToolbar()
 		.Padding(4.f)
 		.ToolTipText(LOCTEXT("SyncSelectionButton_Tooltip", "Sync Selection\nIf enabled, clicking an item in the mixer list will also select the item in the Scene Outliner.\nAlt + Click to select items in mixer without selecting the item in the Scene outliner.\nIf disabled, selections will not sync unless Alt is held. Effectively, this is the opposite behavior."))
 		.Style(&FAppStyle::Get().GetWidgetStyle<FCheckBoxStyle>("ToggleButtonCheckbox"))
-		.ForegroundColor(FSlateColor::UseForeground())
 		.IsChecked_Lambda([]()
 		{
 			if (const UObjectMixerEditorProjectSettings* Settings = GetDefault<UObjectMixerEditorProjectSettings>())
diff --git a/Engine/Plugins/Editor/ObjectMixer/ObjectMixer/Source/ObjectMixer/Private/Views/Widgets/SCollectionSelectionButton.cpp b/Engine/Plugins/Editor/ObjectMixer/ObjectMixer/Source/ObjectMixer/Private/Views/Widgets/SCollectionSelectionButton.cpp
--- a/Engine/Plugins/Editor/ObjectMixer/ObjectMixer/Source/ObjectMixer/Private/Views/Widgets/SCollectionSelectionButton.cpp
+++ b/Engine/Plugins/Editor/ObjectMixer/ObjectMixer/Source/ObjectMixer/Private/Views/Widgets/SCollectionSelectionButton.cpp
@@ -49,11 +49,13 @@ void SCollectionSelectionButton::Construct(const FArguments& InArgs, const TShar
 		SNew(SBorder)
 		.Padding(FMargin(16, 4))
 		.BorderImage(this, &SCollectionSelectionButton::GetBorderBrush)
+		.ForegroundColor(this, &SCollectionSelectionButton::GetBorderForeground)
 		[
 			SNew(STextBlock)
 			.TextStyle(FAppStyle::Get(), "SmallText")
 			.Text(FText::FromName(CollectionName))
 			.Visibility(EVisibility::HitTestInvisible)
+			.ColorAndOpacity(FSlateColor::UseForeground())
 		]
 	];
 }
@@ -212,10 +214,6 @@ const FSlateBrush* SCollectionSelectionButton::GetBorderBrush() const
 		{
 			return &CheckedPressedImage;
 		}
-		else if (bDropIsValid)
-		{
-			return &CheckedValidDropImage;
-		}
 		else if (IsHovered())
 		{
 			return &CheckedHoveredImage;
@@ -228,10 +226,6 @@ const FSlateBrush* SCollectionSelectionButton::GetBorderBrush() const
 	{
 		return &UncheckedPressedImage;
 	}
-	else if (bDropIsValid)
-	{
-		return &UncheckedValidDropImage;
-	}
 	else if (IsHovered())
 	{
 		return &UncheckedHoveredImage;
@@ -241,4 +235,15 @@ const FSlateBrush* SCollectionSelectionButton::GetBorderBrush() const
 			
 }
 
+FSlateColor SCollectionSelectionButton::GetBorderForeground() const
+{
+	if (MainPanelPtr.Pin()->IsCollectionChecked(CollectionName) == ECheckBoxState::Checked ||
+		bIsPressed || IsHovered())
+	{
+		return FStyleColors::White;
+	}
+
+	return FStyleColors::Foreground;
+}
+
 #undef LOCTEXT_NAMESPACE
diff --git a/Engine/Plugins/Editor/ObjectMixer/ObjectMixer/Source/ObjectMixer/Public/Views/List/SObjectMixerEditorListRow.h b/Engine/Plugins/Editor/ObjectMixer/ObjectMixer/Source/ObjectMixer/Public/Views/List/SObjectMixerEditorListRow.h
--- a/Engine/Plugins/Editor/ObjectMixer/ObjectMixer/Source/ObjectMixer/Public/Views/List/SObjectMixerEditorListRow.h
+++ b/Engine/Plugins/Editor/ObjectMixer/ObjectMixer/Source/ObjectMixer/Public/Views/List/SObjectMixerEditorListRow.h
@@ -47,8 +47,6 @@ private:
 	/** Get the brush for this widget */
 	const FSlateBrush* GetVisibilityBrush() const;
 
-	static const FSlateBrush* GetBorderImage(const FObjectMixerEditorListRow::EObjectMixerEditorListRowType InRowType);
-
 	TSharedPtr<SWidget> GenerateCells(const FName& InColumnName, const TSharedPtr<FObjectMixerEditorListRow> PinnedItem);
 
 	void OnPropertyChanged(const FProperty* Property, void* ContainerWithChangedProperty);
diff --git a/Engine/Plugins/Editor/ObjectMixer/ObjectMixer/Source/ObjectMixer/Public/Views/Widgets/SCollectionSelectionButton.h b/Engine/Plugins/Editor/ObjectMixer/ObjectMixer/Source/ObjectMixer/Public/Views/Widgets/SCollectionSelectionButton.h
--- a/Engine/Plugins/Editor/ObjectMixer/ObjectMixer/Source/ObjectMixer/Public/Views/Widgets/SCollectionSelectionButton.h
+++ b/Engine/Plugins/Editor/ObjectMixer/ObjectMixer/Source/ObjectMixer/Public/Views/Widgets/SCollectionSelectionButton.h
@@ -32,6 +32,7 @@ public:
 	virtual FReply OnDrop(const FGeometry& MyGeometry, const FDragDropEvent& DragDropEvent) override;
 
 	const FSlateBrush* GetBorderBrush() const;
+	FSlateColor GetBorderForeground() const;
 
 	TWeakPtr<SObjectMixerEditorMainPanel> MainPanelPtr;
 	FName CollectionName = NAME_None;
