Object Mixer:

-Fix for: Add Light not reflected in panel

diff --git a/Engine/Plugins/Editor/ObjectMixer/ObjectMixer/Source/ObjectMixer/Private/Views/Widgets/SObjectMixerPlacementAssetMenuEntry.cpp b/Engine/Plugins/Editor/ObjectMixer/ObjectMixer/Source/ObjectMixer/Private/Views/Widgets/SObjectMixerPlacementAssetMenuEntry.cpp
--- a/Engine/Plugins/Editor/ObjectMixer/ObjectMixer/Source/ObjectMixer/Private/Views/Widgets/SObjectMixerPlacementAssetMenuEntry.cpp
+++ b/Engine/Plugins/Editor/ObjectMixer/ObjectMixer/Source/ObjectMixer/Private/Views/Widgets/SObjectMixerPlacementAssetMenuEntry.cpp
@@ -6,6 +6,7 @@
 #include "ClassIconFinder.h"
 #include "DragAndDrop/AssetDragDropOp.h"
 #include "EditorClassUtils.h"
+#include "LevelEditorActions.h"
 #include "LevelEditorViewport.h"
 #include "PlacementMode/Public/IPlacementModeModule.h"
 #include "Styling/SlateIconFinder.h"
@@ -169,7 +170,7 @@ FReply SObjectMixerPlacementAssetMenuEntry::OnMouseButtonUp(const FGeometry& MyG
 				FActorFactoryAssetProxy::GetFactoryForAssetObject(ClassObject);
 			}
 		}
-		NewActor = GEditor->GetEditorWorldContext().World()->SpawnActor(Factory->NewActorClass);
+		NewActor = FLevelEditorActionCallbacks::AddActor(Factory, Item->AssetData, nullptr);
 		if (NewActor && GCurrentLevelEditingViewportClient)
 		{
   			GEditor->MoveActorInFrontOfCamera(*NewActor, 
