Object Mixer:

-Fix for: (Component) label appended to light name

diff --git a/Engine/Plugins/Editor/ObjectMixer/ObjectMixer/Source/ObjectMixer/Private/ObjectFilter/ObjectMixerEditorObjectFilter.cpp b/Engine/Plugins/Editor/ObjectMixer/ObjectMixer/Source/ObjectMixer/Private/ObjectFilter/ObjectMixerEditorObjectFilter.cpp
--- a/Engine/Plugins/Editor/ObjectMixer/ObjectMixer/Source/ObjectMixer/Private/ObjectFilter/ObjectMixerEditorObjectFilter.cpp
+++ b/Engine/Plugins/Editor/ObjectMixer/ObjectMixer/Source/ObjectMixer/Private/ObjectFilter/ObjectMixerEditorObjectFilter.cpp
@@ -5,23 +5,53 @@
 #include "GameFramework/Actor.h"
 #include "Kismet2/ComponentEditorUtils.h"
 
-FText UObjectMixerObjectFilter::GetRowDisplayName(UObject* InObject, EObjectMixerTreeViewMode InViewMode) const
+FText UObjectMixerObjectFilter::GetRowDisplayName(UObject* InObject, const bool bIsHybridRow) const
 {
 	if (IsValid(InObject))
 	{
-		if (InViewMode == EObjectMixerTreeViewMode::NoFolders || InViewMode == EObjectMixerTreeViewMode::Folders)
+		if (const AActor* AsActor = Cast<AActor>(InObject))
+		{
+			return FText::FromString(AsActor->GetActorLabel());
+		}
+
+		if (bIsHybridRow)
 		{
 			if (const AActor* Outer = InObject->GetTypedOuter<AActor>())
 			{
-				return FText::Format(INVTEXT("{0} ({1})"), FText::FromString(Outer->GetActorLabel()), FText::FromString(InObject->GetName()));
+				return FText::FromString(Outer->GetActorLabel());
+			}
+		}
+		
+		if (const UActorComponent* AsActorComponent = Cast<UActorComponent>(InObject))
+		{
+			const FName ComponentName = FComponentEditorUtils::FindVariableNameGivenComponentInstance(AsActorComponent);
+			if (ComponentName != NAME_None)
+			{
+				return FText::FromName(ComponentName);
 			}
 		}
+		
+		return FText::FromString(InObject->GetName());
+	}
 
+	return FText::GetEmpty();
+}
+
+FText UObjectMixerObjectFilter::GetRowTooltipText(UObject* InObject, const bool bIsHybridRow) const
+{
+	if (IsValid(InObject))
+	{
+		if (const AActor* Outer = InObject->GetTypedOuter<AActor>())
+		{
+			return FText::Format(INVTEXT("{0} ({1})"), FText::FromString(Outer->GetActorLabel()), FText::FromString(InObject->GetName()));
+		}
+		
 		if (const AActor* AsActor = Cast<AActor>(InObject))
 		{
 			return FText::FromString(AsActor->GetActorLabel());
 		}
-		else if (const UActorComponent* AsActorComponent = Cast<UActorComponent>(InObject))
+		
+		if (const UActorComponent* AsActorComponent = Cast<UActorComponent>(InObject))
 		{
 			const FName ComponentName = FComponentEditorUtils::FindVariableNameGivenComponentInstance(AsActorComponent);
 			if (ComponentName != NAME_None)
@@ -45,7 +75,7 @@ FText UObjectMixerObjectFilter::GetRowDisplayName(UObject* InObject, EObjectMixe
 	return FText::GetEmpty();
 }
 
-bool UObjectMixerObjectFilter::GetRowEditorVisibility(UObject* InObject, EObjectMixerTreeViewMode InViewMode) const
+bool UObjectMixerObjectFilter::GetRowEditorVisibility(UObject* InObject) const
 {
 	if (IsValid(InObject))
 	{
@@ -62,7 +92,7 @@ bool UObjectMixerObjectFilter::GetRowEditorVisibility(UObject* InObject, EObject
 	return false;
 }
 
-void UObjectMixerObjectFilter::OnSetRowEditorVisibility(UObject* InObject, bool bNewIsVisible, EObjectMixerTreeViewMode InViewMode) const
+void UObjectMixerObjectFilter::OnSetRowEditorVisibility(UObject* InObject, bool bNewIsVisible) const
 {
 	if (IsValid(InObject))
 	{
diff --git a/Engine/Plugins/Editor/ObjectMixer/ObjectMixer/Source/ObjectMixer/Private/Views/List/ObjectMixerEditorListRow.cpp b/Engine/Plugins/Editor/ObjectMixer/ObjectMixer/Source/ObjectMixer/Private/Views/List/ObjectMixerEditorListRow.cpp
--- a/Engine/Plugins/Editor/ObjectMixer/ObjectMixer/Source/ObjectMixer/Private/Views/List/ObjectMixerEditorListRow.cpp
+++ b/Engine/Plugins/Editor/ObjectMixer/ObjectMixer/Source/ObjectMixer/Private/Views/List/ObjectMixerEditorListRow.cpp
@@ -218,7 +218,7 @@ bool FObjectMixerEditorListRow::MatchSearchTokensToSearchTerms(
 		{			
 			if (const UObjectMixerObjectFilter* Filter = GetObjectFilter())
 			{
-				CachedSearchTerms += Filter->GetRowDisplayName(Object, GetTreeViewMode()).ToString();
+				CachedSearchTerms += Filter->GetRowDisplayName(Object).ToString();
 			}
 		}
 	}
@@ -333,7 +333,7 @@ bool FObjectMixerEditorListRow::HasVisibleChildRowWidgets() const
 	return false;
 }
 
-FText FObjectMixerEditorListRow::GetDisplayName()
+FText FObjectMixerEditorListRow::GetDisplayName(const bool bIsHybridRow) const
 {
 	if (!GetDisplayNameOverride().IsEmpty())
 	{
@@ -344,7 +344,7 @@ FText FObjectMixerEditorListRow::GetDisplayName()
 	{
 		if (const TObjectPtr<UObject> Object = GetObject())
 		{
-			return Filter->GetRowDisplayName(Object, GetTreeViewMode());
+			return Filter->GetRowDisplayName(Object, bIsHybridRow);
 		}
 	}
 
@@ -427,7 +427,7 @@ bool FObjectMixerEditorListRow::GetObjectVisibility()
 	
 	if (const UObjectMixerObjectFilter* Filter = GetObjectFilter())
 	{
-		return Filter->GetRowEditorVisibility(GetObject(), GetTreeViewMode());
+		return Filter->GetRowEditorVisibility(GetObject());
 	}
 
 	return false;
@@ -437,7 +437,7 @@ void FObjectMixerEditorListRow::SetObjectVisibility(const bool bNewIsVisible, co
 {
 	if (const UObjectMixerObjectFilter* Filter = GetObjectFilter())
 	{
-		Filter->OnSetRowEditorVisibility(GetObject(), bNewIsVisible, GetTreeViewMode());
+		Filter->OnSetRowEditorVisibility(GetObject(), bNewIsVisible);
 
 		if (bIsRecursive)
 		{
diff --git a/Engine/Plugins/Editor/ObjectMixer/ObjectMixer/Source/ObjectMixer/Private/Views/List/SObjectMixerEditorListRow.cpp b/Engine/Plugins/Editor/ObjectMixer/ObjectMixer/Source/ObjectMixer/Private/Views/List/SObjectMixerEditorListRow.cpp
--- a/Engine/Plugins/Editor/ObjectMixer/ObjectMixer/Source/ObjectMixer/Private/Views/List/SObjectMixerEditorListRow.cpp
+++ b/Engine/Plugins/Editor/ObjectMixer/ObjectMixer/Source/ObjectMixer/Private/Views/List/SObjectMixerEditorListRow.cpp
@@ -62,13 +62,15 @@ void SObjectMixerEditorListRow::Construct(
 
 TSharedRef<SWidget> SObjectMixerEditorListRow::GenerateWidgetForColumn(const FName& InColumnName)
 {
+	check(Item.IsValid());
+
 	const FObjectMixerEditorListRowPtr PinnedItem =
 		HybridRowIndex != INDEX_NONE ? Item.Pin()->GetChildRows()[HybridRowIndex] : Item.Pin();
 
 	if (const TSharedPtr<SWidget> CellWidget = GenerateCells(InColumnName, PinnedItem))
 	{
 		if (InColumnName == SObjectMixerEditorList::ItemNameColumnName)
-		{
+		{			
 			// The first column gets the tree expansion arrow for this row
 			return SNew(SBox)
 				.MinDesiredHeight(20)
@@ -295,7 +297,11 @@ TSharedPtr<SWidget> SObjectMixerEditorListRow::GenerateCells(
 {
 	check(PinnedItem.IsValid());
 	
-	const bool bIsRowTypeNone = PinnedItem->GetRowType() == FObjectMixerEditorListRow::None;
+	if (PinnedItem->GetRowType() == FObjectMixerEditorListRow::None)
+	{
+		return SNullWidget::NullWidget;
+	}
+	
 	const bool bIsHybridRow = HybridRowIndex != INDEX_NONE;
 	
 	if (InColumnName.IsEqual(SObjectMixerEditorList::ItemNameColumnName))
@@ -314,7 +320,8 @@ TSharedPtr<SWidget> SObjectMixerEditorListRow::GenerateCells(
 		];
 
 		bool bNeedsStandardTextBlock = true;
-		const FText DisplayName = PinnedItem->GetDisplayName();
+		const FText DisplayName = PinnedItem->GetDisplayName(bIsHybridRow);
+		
 		if (TObjectPtr<UObject> Object = PinnedItem->GetObject())
 		{
 			if (UClass* ActorClass = Object->GetClass())
@@ -340,6 +347,16 @@ TSharedPtr<SWidget> SObjectMixerEditorListRow::GenerateCells(
 
 		if (bNeedsStandardTextBlock)
 		{
+			FText TooltipText = DisplayName;
+
+			if (const UObjectMixerObjectFilter* Filter = PinnedItem->GetObjectFilter())
+			{
+				if (const TObjectPtr<UObject> Object = PinnedItem->GetObject())
+				{
+					TooltipText = Filter->GetRowTooltipText(Object, bIsHybridRow);
+				}
+			}
+			
 			HBox->AddSlot()
 			.Padding(FMargin(10.0, 0, 0, 0))
 			[
@@ -347,7 +364,7 @@ TSharedPtr<SWidget> SObjectMixerEditorListRow::GenerateCells(
 				.Visibility(EVisibility::Visible)
 				.Justification(ETextJustify::Left)
 				.Text(DisplayName)
-				.ToolTipText(DisplayName)
+				.ToolTipText(TooltipText)
 			];
 		}
 		
diff --git a/Engine/Plugins/Editor/ObjectMixer/ObjectMixer/Source/ObjectMixer/Public/ObjectFilter/ObjectMixerEditorObjectFilter.h b/Engine/Plugins/Editor/ObjectMixer/ObjectMixer/Source/ObjectMixer/Public/ObjectFilter/ObjectMixerEditorObjectFilter.h
--- a/Engine/Plugins/Editor/ObjectMixer/ObjectMixer/Source/ObjectMixer/Public/ObjectFilter/ObjectMixerEditorObjectFilter.h
+++ b/Engine/Plugins/Editor/ObjectMixer/ObjectMixer/Source/ObjectMixer/Public/ObjectFilter/ObjectMixerEditorObjectFilter.h
@@ -59,29 +59,30 @@ public:
 	/**
 	 * Get the text to display for the object name/label.
 	 * This is useful if one of your classes is a component type and you want the label of the component's owning actor, for example.
-	 * If not overridden, this returns the object's name.
-	 * @param InViewMode The tree view display option (i.e. Flat list or hierarchy).
+	 * If not overridden, this returns the actor's label if it's an actor and the object's name for all other objects.
+	 * @param bIsHybridRow If true, this row represents a single matching component condensed into a single row with its actor parent.
 	 */
-	virtual FText GetRowDisplayName(
-		UObject* InObject, EObjectMixerTreeViewMode InViewMode) const;
+	virtual FText GetRowDisplayName(UObject* InObject, const bool bIsHybridRow = false) const;
+
+	/**
+	 * Get the text to display for the row's tooltip when hovering over it.
+	 * @param bIsHybridRow If true, this row represents a single matching component condensed into a single row with its actor parent. (unused)
+	 */
+	virtual FText GetRowTooltipText(UObject* InObject, const bool bIsHybridRow = false) const;
 
 	/**
 	 * Controls how to display the row's visibility icon. Return true if the object should be visible.
 	 * Generally this should work like the Scene Outliner does.
 	 * If not overridden, we use the Editor Visibility of the object's AActor outer (unless it's an actor itself).
-	 * @param InViewMode The tree view display option (i.e. Flat list or hierarchy) (unused)
 	 */
-	virtual bool GetRowEditorVisibility(
-		UObject* InObject, EObjectMixerTreeViewMode InViewMode) const;
+	virtual bool GetRowEditorVisibility(UObject* InObject) const;
 
 	/**
 	 * Controls what happens when the row's visibility icon is clicked.
 	 * Generally this should work like the Scene Outliner does.
 	 * If not overridden, we set the Editor Visibility of the object's AActor outer (unless it's an actor itself).
-	 * @param InViewMode The tree view display option (i.e. Flat list or hierarchy) (unused)
 	 */
-	virtual void OnSetRowEditorVisibility(
-		UObject* InObject, bool bNewIsVisible, EObjectMixerTreeViewMode InViewMode) const;
+	virtual void OnSetRowEditorVisibility(UObject* InObject, bool bNewIsVisible) const;
 
 	/**
 	 * Specify a list of property names corresponding to columns you want to show by default.
@@ -185,46 +186,53 @@ public:
 	/**
 	 * Get the text to display for the object name/label.
 	 * This is useful if one of your classes is a component type and you want the label of the component's owning actor, for example.
-	 * If not overridden, this returns the object's name.
-	 * @param InViewMode The tree view display option (i.e. Flat list or hierarchy).
+	 * If not overridden, this returns the actor's label if it's an actor and the object's name for all other objects.
+	 * @param bIsHybridRow If true, this row represents a single matching component condensed into a single row with its actor parent.
+	 */
+	UFUNCTION(BlueprintNativeEvent, BlueprintCallable, Category = "Object Mixer")
+	FText GetRowDisplayName(UObject* InObject, const bool bIsHybridRow) const override;
+	
+	FText GetRowDisplayName_Implementation(UObject* InObject, const bool bIsHybridRow) const
+	{
+		return Super::GetRowDisplayName(InObject, bIsHybridRow);
+	}
+
+	/**
+	 * Get the text to display for the row's tooltip when hovering over it.
+	 * @param bIsHybridRow If true, this row represents a single matching component condensed into a single row with its actor parent. (unused)
 	 */
 	UFUNCTION(BlueprintNativeEvent, BlueprintCallable, Category = "Object Mixer")
-	FText GetRowDisplayName(
-		UObject* InObject, EObjectMixerTreeViewMode InViewMode = EObjectMixerTreeViewMode::Folders) const override;
+	FText GetRowTooltipText(UObject* InObject, const bool bIsHybridRow) const override;
 	
-	FText GetRowDisplayName_Implementation(UObject* InObject, EObjectMixerTreeViewMode InViewMode) const
+	FText GetRowTooltipText_Implementation(UObject* InObject, const bool bIsHybridRow) const
 	{
-		return Super::GetRowDisplayName(InObject, InViewMode);
+		return Super::GetRowTooltipText(InObject, bIsHybridRow);
 	}
 	
 	/**
 	 * Controls how to display the row's visibility icon. Return true if the object should be visible.
 	 * Generally this should work like the Scene Outliner does.
 	 * If not overridden, we use the Editor Visibility of the object's AActor outer (unless it's an actor itself).
-	 * @param InViewMode The tree view display option (i.e. Flat list or hierarchy) (unused)
 	 */
 	UFUNCTION(BlueprintNativeEvent, BlueprintCallable, Category = "Object Mixer")
-	bool GetRowEditorVisibility(
-		UObject* InObject, EObjectMixerTreeViewMode InViewMode = EObjectMixerTreeViewMode::Folders) const override;
+	bool GetRowEditorVisibility(UObject* InObject) const override;
 
-	bool GetRowEditorVisibility_Implementation(UObject* InObject, EObjectMixerTreeViewMode InViewMode) const
+	bool GetRowEditorVisibility_Implementation(UObject* InObject) const
 	{
-		return Super::GetRowEditorVisibility(InObject, InViewMode);
+		return Super::GetRowEditorVisibility(InObject);
 	}
 	
 	/**
 	 * Controls what happens when the row's visibility icon is clicked.
 	 * Generally this should work like the Scene Outliner does.
 	 * If not overridden, we set the Editor Visibility of the object's AActor outer (unless it's an actor itself).
-	 * @param InViewMode The tree view display option (i.e. Flat list or hierarchy) (unused)
 	 */
 	UFUNCTION(BlueprintNativeEvent, BlueprintCallable, Category = "Object Mixer")
-	void OnSetRowEditorVisibility(
-		UObject* InObject, bool bNewIsVisible, EObjectMixerTreeViewMode InViewMode = EObjectMixerTreeViewMode::Folders) const override;
+	void OnSetRowEditorVisibility(UObject* InObject, bool bNewIsVisible) const override;
 
-	void OnSetRowEditorVisibility_Implementation(UObject* InObject, bool bNewIsVisible, EObjectMixerTreeViewMode InViewMode) const
+	void OnSetRowEditorVisibility_Implementation(UObject* InObject, bool bNewIsVisible) const
 	{
-		return Super::OnSetRowEditorVisibility(InObject, bNewIsVisible, InViewMode);
+		return Super::OnSetRowEditorVisibility(InObject, bNewIsVisible);
 	}
 	
 	/**
diff --git a/Engine/Plugins/Editor/ObjectMixer/ObjectMixer/Source/ObjectMixer/Public/Views/List/ObjectMixerEditorListRow.h b/Engine/Plugins/Editor/ObjectMixer/ObjectMixer/Source/ObjectMixer/Public/Views/List/ObjectMixerEditorListRow.h
--- a/Engine/Plugins/Editor/ObjectMixer/ObjectMixer/Source/ObjectMixer/Public/Views/List/ObjectMixerEditorListRow.h
+++ b/Engine/Plugins/Editor/ObjectMixer/ObjectMixer/Source/ObjectMixer/Public/Views/List/ObjectMixerEditorListRow.h
@@ -132,7 +132,7 @@ struct OBJECTMIXEREDITOR_API FObjectMixerEditorListRow final : TSharedFromThis<F
 
 	[[nodiscard]] bool HasVisibleChildRowWidgets() const;
 
-	[[nodiscard]] FText GetDisplayName();
+	[[nodiscard]] FText GetDisplayName(const bool bIsHybridRow = false) const;
 
 	[[nodiscard]] const FText& GetDisplayNameOverride() const
 	{
