Object Mixer:

-Visual adjustments to light mixer panel
- Fix: Solo icon is incorrect when toggled
- Fix: Solo icon not appearing for hybrid rows
- Fix: Eye icon misaligned with header
-Fix: Selection in viewport/outliner should be reflected in Light Mixer when sync selection is enabled
-Restore Selected rows on list refresh
-Fix: Drag & Drop (Reorder and Reparent) not working properly

diff --git a/Engine/Plugins/Editor/ObjectMixer/ObjectMixer/Source/ObjectMixer/Private/ObjectMixerEditorModule.cpp b/Engine/Plugins/Editor/ObjectMixer/ObjectMixer/Source/ObjectMixer/Private/ObjectMixerEditorModule.cpp
--- a/Engine/Plugins/Editor/ObjectMixer/ObjectMixer/Source/ObjectMixer/Private/ObjectMixerEditorModule.cpp
+++ b/Engine/Plugins/Editor/ObjectMixer/ObjectMixer/Source/ObjectMixer/Private/ObjectMixerEditorModule.cpp
@@ -12,6 +12,7 @@
 #include "ISettingsModule.h"
 #include "LevelEditor.h"
 #include "ObjectMixerEditorSerializedData.h"
+#include "Selection.h"
 #include "Misc/TransactionObjectEvent.h"
 #include "Modules/ModuleManager.h"
 #include "ToolMenus.h"
@@ -138,6 +139,14 @@ void FObjectMixerEditorModule::RefreshList() const
 	}
 }
 
+void FObjectMixerEditorModule::RequestSyncEditorSelectionToListSelection()
+{
+	if (MainPanel.IsValid())
+	{
+		MainPanel->RequestSyncEditorSelectionToListSelection();
+	}
+}
+
 void FObjectMixerEditorModule::RegisterMenuGroup()
 {
 	WorkspaceGroup = WorkspaceMenu::GetMenuStructure().GetLevelEditorCategory()->AddGroup(
@@ -269,10 +278,6 @@ TSharedRef<SDockTab> FObjectMixerEditorModule::SpawnMainPanelTab()
 void FObjectMixerEditorModule::BindDelegates()
 {
 	check(GEngine && GEditor);
-	DelegateHandles.Add(GEngine->OnLevelActorFolderChanged().AddLambda([this] (const AActor*, FName)
-	{
-		RequestRebuildList();
-	}));
 	DelegateHandles.Add(GEngine->OnLevelActorListChanged().AddLambda([this] ()
 	{
 		RequestRebuildList();
@@ -281,19 +286,19 @@ void FObjectMixerEditorModule::BindDelegates()
 	{
 		RequestRebuildList();
 	}));
-	DelegateHandles.Add(GEngine->OnLevelActorDetached().AddLambda([this] (AActor*, const AActor*)
+	DelegateHandles.Add(GEngine->OnActorFoldersUpdatedEvent().AddLambda([this] (ULevel*)
 	{
 		RequestRebuildList();
 	}));
-	DelegateHandles.Add(GEngine->OnLevelActorAttached().AddLambda([this] (AActor*, const AActor*)
+	DelegateHandles.Add(GEngine->OnLevelActorFolderChanged().AddLambda([this] (const AActor*, FName)
 	{
 		RequestRebuildList();
 	}));
-	DelegateHandles.Add(GEngine->OnActorFolderAdded().AddLambda([this] (UActorFolder*)
+	DelegateHandles.Add(GEngine->OnLevelActorAttached().AddLambda([this] (AActor*, const AActor*)
 	{
 		RequestRebuildList();
 	}));
-	DelegateHandles.Add(GEngine->OnActorFoldersUpdatedEvent().AddLambda([this] (ULevel*)
+	DelegateHandles.Add(GEngine->OnLevelActorDetached().AddLambda([this] (AActor*, const AActor*)
 	{
 		RequestRebuildList();
 	}));
@@ -302,8 +307,16 @@ void FObjectMixerEditorModule::BindDelegates()
 	{
 		RequestRebuildList();
 	}));
+	DelegateHandles.Add(GEditor->OnLevelActorDeleted().AddLambda([this] (AActor*)
+	{
+		RequestRebuildList();
+	}));
+	DelegateHandles.Add(GEditor->GetSelectedActors()->SelectionChangedEvent.AddLambda([this] (UObject*)
+	{
+		RequestSyncEditorSelectionToListSelection();
+	}));
 
-	DelegateHandles.Add(FEditorDelegates::MapChange.AddLambda([this](uint32 Index)
+	DelegateHandles.Add(FEditorDelegates::MapChange.AddLambda([this](uint32)
 	{
 		RequestRebuildList();
 	}));
@@ -318,7 +331,7 @@ void FObjectMixerEditorModule::BindDelegates()
 
 	DelegateHandles.Add(FCoreUObjectDelegates::OnObjectTransacted.AddLambda([this](UObject*, const FTransactionObjectEvent& Event)
 	{
-		if (Event.GetEventType() == ETransactionObjectEventType::Finalized)
+		if (Event.GetEventType() == ETransactionObjectEventType::Finalized && Event.HasPropertyChanges())
 		{
 			RequestRebuildList();
 		}
diff --git a/Engine/Plugins/Editor/ObjectMixer/ObjectMixer/Source/ObjectMixer/Private/Views/List/ObjectMixerEditorList.cpp b/Engine/Plugins/Editor/ObjectMixer/ObjectMixer/Source/ObjectMixer/Private/Views/List/ObjectMixerEditorList.cpp
--- a/Engine/Plugins/Editor/ObjectMixer/ObjectMixer/Source/ObjectMixer/Private/Views/List/ObjectMixerEditorList.cpp
+++ b/Engine/Plugins/Editor/ObjectMixer/ObjectMixer/Source/ObjectMixer/Private/Views/List/ObjectMixerEditorList.cpp
@@ -50,15 +50,29 @@ void FObjectMixerEditorList::RefreshList() const
 	}
 }
 
+void FObjectMixerEditorList::RequestSyncEditorSelectionToListSelection() const
+{
+	if (ListWidget.IsValid())
+	{
+		ListWidget->RequestSyncEditorSelectionToListSelection();
+	}
+}
+
 void FObjectMixerEditorList::ExecuteListViewSearchOnAllRows(const FString& SearchString,
-	const bool bShouldRefreshAfterward)
+                                                            const bool bShouldRefreshAfterward)
 {
-	ListWidget->ExecuteListViewSearchOnAllRows(SearchString, bShouldRefreshAfterward);
+	if (ListWidget.IsValid())
+	{
+		ListWidget->ExecuteListViewSearchOnAllRows(SearchString, bShouldRefreshAfterward);
+	}
 }
 
 void FObjectMixerEditorList::EvaluateIfRowsPassFilters(const bool bShouldRefreshAfterward) const
 {
-	ListWidget->EvaluateIfRowsPassFilters();
+	if (ListWidget.IsValid())
+	{
+		ListWidget->EvaluateIfRowsPassFilters();
+	}
 }
 
 TWeakPtr<FObjectMixerEditorMainPanel> FObjectMixerEditorList::GetMainPanelModel()
@@ -66,17 +80,22 @@ TWeakPtr<FObjectMixerEditorMainPanel> FObjectMixerEditorList::GetMainPanelModel(
 	return MainPanelModelPtr;
 }
 
-TWeakPtr<FObjectMixerEditorListRow> FObjectMixerEditorList::GetSoloRow()
+TSet<TWeakPtr<FObjectMixerEditorListRow>> FObjectMixerEditorList::GetSoloRows()
+{
+	return GetMainPanelModel().Pin()->GetSoloRows();
+}
+
+void FObjectMixerEditorList::AddSoloRow(TSharedRef<FObjectMixerEditorListRow> InRow)
 {
-	return GetMainPanelModel().Pin()->GetSoloRow();
+	GetMainPanelModel().Pin()->AddSoloRow(InRow);
 }
 
-void FObjectMixerEditorList::SetSoloRow(TSharedRef<FObjectMixerEditorListRow> InRow)
+void FObjectMixerEditorList::RemoveSoloRow(TSharedRef<FObjectMixerEditorListRow> InRow)
 {
-	GetMainPanelModel().Pin()->SetSoloRow(InRow);
+	GetMainPanelModel().Pin()->RemoveSoloRow(InRow);
 }
 
-void FObjectMixerEditorList::ClearSoloRow()
+void FObjectMixerEditorList::ClearSoloRows()
 {
-	GetMainPanelModel().Pin()->ClearSoloRow();
+	GetMainPanelModel().Pin()->ClearSoloRows();
 }
diff --git a/Engine/Plugins/Editor/ObjectMixer/ObjectMixer/Source/ObjectMixer/Private/Views/List/ObjectMixerEditorListRow.cpp b/Engine/Plugins/Editor/ObjectMixer/ObjectMixer/Source/ObjectMixer/Private/Views/List/ObjectMixerEditorListRow.cpp
--- a/Engine/Plugins/Editor/ObjectMixer/ObjectMixer/Source/ObjectMixer/Private/Views/List/ObjectMixerEditorListRow.cpp
+++ b/Engine/Plugins/Editor/ObjectMixer/ObjectMixer/Source/ObjectMixer/Private/Views/List/ObjectMixerEditorListRow.cpp
@@ -449,19 +449,26 @@ void FObjectMixerEditorListRow::SetObjectVisibility(const bool bNewIsVisible, co
 	}
 }
 
-bool FObjectMixerEditorListRow::IsThisRowSolo() const
+bool FObjectMixerEditorListRow::IsThisRowSolo()
 {
-	return GetListViewPtr().Pin()->GetSoloRow().HasSameObject(this);
+	return GetListViewPtr().Pin()->GetSoloRows().Contains(SharedThis(this));
 }
 
-void FObjectMixerEditorListRow::SetThisAsSoloRow()
+void FObjectMixerEditorListRow::SetRowSoloState(const bool bNewSolo)
 {
-	GetListViewPtr().Pin()->SetSoloRow(SharedThis(this));
+	if (bNewSolo)
+	{
+		GetListViewPtr().Pin()->AddSoloRow(SharedThis(this));
+	}
+	else
+	{
+		GetListViewPtr().Pin()->RemoveSoloRow(SharedThis(this));
+	}
 }
 
-void FObjectMixerEditorListRow::ClearSoloRow()
+void FObjectMixerEditorListRow::ClearSoloRows() const
 {
-	GetListViewPtr().Pin()->ClearSoloRow();
+	GetListViewPtr().Pin()->ClearSoloRows();
 }
 
 FObjectMixerEditorListRowPtr FObjectMixerEditorListRow::GetAsShared()
diff --git a/Engine/Plugins/Editor/ObjectMixer/ObjectMixer/Source/ObjectMixer/Private/Views/List/SObjectMixerEditorList.cpp b/Engine/Plugins/Editor/ObjectMixer/ObjectMixer/Source/ObjectMixer/Private/Views/List/SObjectMixerEditorList.cpp
--- a/Engine/Plugins/Editor/ObjectMixer/ObjectMixer/Source/ObjectMixer/Private/Views/List/SObjectMixerEditorList.cpp
+++ b/Engine/Plugins/Editor/ObjectMixer/ObjectMixer/Source/ObjectMixer/Private/Views/List/SObjectMixerEditorList.cpp
@@ -86,13 +86,16 @@ void SObjectMixerEditorList::Construct(const FArguments& InArgs, TSharedRef<FObj
 				})
 				.OnSelectionChanged_Lambda([this] (const FObjectMixerEditorListRowPtr& Row, const ESelectInfo::Type SelectionType)
 				{
-					const bool bSyncSelectionEnabled = GetDefault<UObjectMixerEditorProjectSettings>()->bSyncSelection;
-					const bool bIsAltDown = FSlateApplication::Get().GetModifierKeys().IsAltDown();
-					const bool bShouldSyncSelection = bSyncSelectionEnabled ? !bIsAltDown : bIsAltDown;
-						
-					if (bShouldSyncSelection)
+					if (SelectionType != ESelectInfo::Direct) // Don't call if selected from code
 					{
-						SelectedTreeItemsToSelectedInLevelEditor();
+						const bool bSyncSelectionEnabled = GetDefault<UObjectMixerEditorProjectSettings>()->bSyncSelection;
+						const bool bIsAltDown = FSlateApplication::Get().GetModifierKeys().IsAltDown();
+						const bool bShouldSyncSelection = bSyncSelectionEnabled ? !bIsAltDown : bIsAltDown;
+						
+						if (bShouldSyncSelection)
+						{
+							SelectedTreeItemsToSelectedInLevelEditor();
+						}
 					}
 				})
 				.OnMouseButtonDoubleClick_Lambda([this] (FObjectMixerEditorListRowPtr Row)
@@ -166,7 +169,7 @@ SObjectMixerEditorList::~SObjectMixerEditorList()
 
 void SObjectMixerEditorList::RebuildList()
 {
-	bShouldRebuild = false;
+	bIsRebuildRequested = false;
 	
 	GenerateTreeView();
 }
@@ -191,7 +194,7 @@ void SObjectMixerEditorList::RefreshList()
 
 void SObjectMixerEditorList::RequestRebuildList(const FString& InItemToScrollTo)
 {
-	bShouldRebuild = true;
+	bIsRebuildRequested = true;
 }
 
 TArray<FObjectMixerEditorListRowPtr> SObjectMixerEditorList::GetSelectedTreeViewItems() const
@@ -204,6 +207,31 @@ int32 SObjectMixerEditorList::GetSelectedTreeViewItemCount() const
 	return TreeViewPtr->GetSelectedItems().Num();
 }
 
+void SObjectMixerEditorList::SyncEditorSelectionToListSelection()
+{
+	if (GetDefault<UObjectMixerEditorProjectSettings>()->bSyncSelection && !bShouldPauseSyncSelection)
+	{
+		bIsEditorToListSelectionSyncRequested = false;
+		const USelection* SelectedActors = GEditor->GetSelectedActors();
+
+		if (const int32 SelectionCount = SelectedActors->Num(); SelectionCount > 0)
+		{
+			bShouldPauseSyncSelection = true;
+			TreeViewPtr->ClearSelection();
+		
+			for (int32 SelectionItr = 0; SelectionItr < SelectionCount; SelectionItr++)
+			{
+				UObject* SelectedActor = SelectedActors->GetSelectedObject(SelectionItr);
+
+				if (const FObjectMixerEditorListRowPtr* Match = ObjectsToRowsCreated.Find(SelectedActor))
+				{
+					TreeViewPtr->SetItemSelection(*Match, true);
+				}
+			}
+		}
+	}
+}
+
 void SObjectMixerEditorList::SetSelectedTreeViewItemActorsEditorVisible(const bool bNewIsVisible, const bool bIsRecursive)
 {
 	for (const TSharedPtr<FObjectMixerEditorListRow>& SelectedItem : TreeViewPtr->GetSelectedItems())
@@ -229,19 +257,24 @@ void SObjectMixerEditorList::SetTreeViewItems(const TArray<FObjectMixerEditorLis
 	TreeViewPtr->RequestListRefresh();
 }
 
-TWeakPtr<FObjectMixerEditorListRow> SObjectMixerEditorList::GetSoloRow()
+TSet<TWeakPtr<FObjectMixerEditorListRow>> SObjectMixerEditorList::GetSoloRows()
+{
+	return GetListModelPtr().Pin()->GetSoloRows();
+}
+
+void SObjectMixerEditorList::AddSoloRow(TSharedRef<FObjectMixerEditorListRow> InRow)
 {
-	return GetListModelPtr().Pin()->GetSoloRow();
+	GetListModelPtr().Pin()->AddSoloRow(InRow);
 }
 
-void SObjectMixerEditorList::SetSoloRow(TSharedRef<FObjectMixerEditorListRow> InRow)
+void SObjectMixerEditorList::RemoveSoloRow(TSharedRef<FObjectMixerEditorListRow> InRow)
 {
-	GetListModelPtr().Pin()->SetSoloRow(InRow);
+	GetListModelPtr().Pin()->RemoveSoloRow(InRow);
 }
 
-void SObjectMixerEditorList::ClearSoloRow()
+void SObjectMixerEditorList::ClearSoloRows()
 {
-	GetListModelPtr().Pin()->ClearSoloRow();
+	GetListModelPtr().Pin()->ClearSoloRows();
 }
 
 FString SObjectMixerEditorList::GetSearchStringFromSearchInputField() const
@@ -438,10 +471,21 @@ void SObjectMixerEditorList::Tick(const FGeometry& AllottedGeometry, const doubl
 {
 	SCompoundWidget::Tick(AllottedGeometry, InCurrentTime, InDeltaTime);
 
-	if (bShouldRebuild)
+	if (bIsRebuildRequested)
 	{
 		RebuildList();
 	}
+
+	if (bIsEditorToListSelectionSyncRequested)
+	{
+		SyncEditorSelectionToListSelection();
+	}
+
+	if (bShouldPauseSyncSelection)
+	{
+		bShouldPauseSyncSelection = false;
+		bIsEditorToListSelectionSyncRequested = false;
+	}
 }
 
 TSharedRef<SWidget> SObjectMixerEditorList::GenerateHeaderRowContextMenu() const
@@ -755,11 +799,13 @@ void SObjectMixerEditorList::FlushMemory(const bool bShouldKeepMemoryAllocated)
 {
 	if (bShouldKeepMemoryAllocated)
 	{
+		ObjectsToRowsCreated.Reset();
 		TreeViewRootObjects.Reset();
 		VisibleTreeViewObjects.Reset();
 	}
 	else
 	{
+		ObjectsToRowsCreated.Empty();
 		TreeViewRootObjects.Empty();
 		VisibleTreeViewObjects.Empty();
 	}
@@ -788,24 +834,34 @@ void SObjectMixerEditorList::CacheTreeState()
 	{
 		static void RecursivelyCacheTreeState(
 			const TArray<FObjectMixerEditorListRowPtr>& InObjects,
-			TMap<FString, bool>& TreeItemExpansionStateCache,
+			TArray<FTreeItemStateCache>& TreeItemStateCache,
 			TSharedPtr<STreeView<FObjectMixerEditorListRowPtr>> TreeViewPtr)
 		{
 			for (const TSharedPtr<FObjectMixerEditorListRow>& TreeViewItem : InObjects)
 			{
-				const FString ObjectName = TreeViewItem->GetDisplayName().ToString();
+				UObject* RowObject = TreeViewItem->GetObject();
+				const FString RowName =
+					TreeViewItem->GetRowType() == FObjectMixerEditorListRow::Folder ?
+						TreeViewItem->GetFolderPath().ToString() : TreeViewItem->GetDisplayName().ToString();
 
-				if (!ObjectName.IsEmpty())
+				if (!RowName.IsEmpty())
 				{
-					TreeItemExpansionStateCache.Add(ObjectName, TreeViewPtr->IsItemExpanded(TreeViewItem));
+					TreeItemStateCache.Add(
+						{
+							RowObject ? RowObject->GetUniqueID() : -1,
+							RowName,
+							TreeViewPtr->IsItemExpanded(TreeViewItem), 
+							TreeViewPtr->IsItemSelected(TreeViewItem)
+						}
+					);
 				}
 
-				RecursivelyCacheTreeState(TreeViewItem->GetChildRows(), TreeItemExpansionStateCache, TreeViewPtr);
+				RecursivelyCacheTreeState(TreeViewItem->GetChildRows(), TreeItemStateCache, TreeViewPtr);
 			}
 		}
 	};
 
-	Local::RecursivelyCacheTreeState(TreeViewRootObjects, TreeItemExpansionStateCache, TreeViewPtr);
+	Local::RecursivelyCacheTreeState(TreeViewRootObjects, TreeItemStateCache, TreeViewPtr);
 }
 
 void SObjectMixerEditorList::RestoreTreeState(const bool bFlushCache)
@@ -814,28 +870,42 @@ void SObjectMixerEditorList::RestoreTreeState(const bool bFlushCache)
 	{
 		static void RecursivelyRestoreTreeState(
 			const TArray<FObjectMixerEditorListRowPtr>& InObjects,
-			TMap<FString, bool>& TreeItemExpansionStateCache,
+			TArray<FTreeItemStateCache>& TreeItemStateCache,
 			TSharedPtr<STreeView<FObjectMixerEditorListRowPtr>> TreeViewPtr,
 			const bool bExpandByDefault)
 		{
 			for (const TSharedPtr<FObjectMixerEditorListRow>& TreeViewItem : InObjects)
 			{
-				const FString ObjectName = TreeViewItem->GetDisplayName().ToString();
-
-				if (!ObjectName.IsEmpty())
-				{
-					if (const bool* bExpansionStatePtr = TreeItemExpansionStateCache.Find(ObjectName))
+				if (const FTreeItemStateCache* StateCachePtr = Algo::FindByPredicate(
+					TreeItemStateCache,
+					[TreeViewItem](const FTreeItemStateCache& Other)
 					{
-						TreeViewPtr->SetItemExpansion(TreeViewItem, *bExpansionStatePtr);
-					}
-					else
-					{
-						TreeViewPtr->SetItemExpansion(TreeViewItem, bExpandByDefault);
+						if (Other.UniqueId != -1)
+						{
+							if (const UObject* RowObject = TreeViewItem->GetObject())
+							{
+								return Other.UniqueId == RowObject->GetUniqueID();
+							}
+						}
+
+						const FString RowName =
+							TreeViewItem->GetRowType() == FObjectMixerEditorListRow::Folder ?
+								TreeViewItem->GetFolderPath().ToString() : TreeViewItem->GetDisplayName().ToString();
+
+						return Other.RowName.Equals(RowName);
 					}
+				))
+				{
+					TreeViewPtr->SetItemExpansion(TreeViewItem, StateCachePtr->bIsExpanded);
+					TreeViewPtr->SetItemSelection(TreeViewItem, StateCachePtr->bIsSelected);
+				}
+				else
+				{
+					TreeViewPtr->SetItemExpansion(TreeViewItem, bExpandByDefault);
 				}
 
 				RecursivelyRestoreTreeState(
-					TreeViewItem->GetChildRows(), TreeItemExpansionStateCache, TreeViewPtr, bExpandByDefault);
+					TreeViewItem->GetChildRows(), TreeItemStateCache, TreeViewPtr, bExpandByDefault);
 			}
 		}
 	};
@@ -847,11 +917,11 @@ void SObjectMixerEditorList::RestoreTreeState(const bool bFlushCache)
 	}
 	
 	Local::RecursivelyRestoreTreeState(
-		TreeViewRootObjects, TreeItemExpansionStateCache, TreeViewPtr, bExpandByDefault);
+		TreeViewRootObjects, TreeItemStateCache, TreeViewPtr, bExpandByDefault);
 
 	if (bFlushCache)
 	{
-		TreeItemExpansionStateCache.Empty();
+		TreeItemStateCache.Empty();
 	}
 }
 
@@ -1044,7 +1114,7 @@ void CreateFolderRowsForTopLevelRow(
 			{
 				FolderRow =
 					MakeShared<FObjectMixerEditorListRow>(
-						nullptr, FObjectMixerEditorListRow::Folder, InListView,
+						BaseActorFolder.GetPath(), FObjectMixerEditorListRow::Folder, InListView,
 						FText::FromName(BaseActorFolder.GetLeafName()));
 
 				FolderMap.Add(BaseActorFolder.GetPath(), FolderRow.ToSharedRef());
@@ -1089,8 +1159,7 @@ void SObjectMixerEditorList::GenerateTreeView()
 			AllMatchingObjects.Add(Object);
 		}
 	});
-
-	TMap<UObject*, FObjectMixerEditorListRowPtr> ObjectsWithRowCreated;
+	
 	TMap<FName, FObjectMixerEditorListRowPtr> FolderMap;
 	for (UObject* Object : AllMatchingObjects)
 	{
@@ -1100,7 +1169,7 @@ void SObjectMixerEditorList::GenerateTreeView()
 
 		CreateOuterRowsForTopLevelRow(
 			NextOuter, TopLevelRow, AllMatchingObjects,
-			ObjectsWithRowCreated, SharedThis(this)
+			ObjectsToRowsCreated, SharedThis(this)
 		);
 		
 		// Now consider folder hierarchy for the top level row's object if desired
@@ -1111,8 +1180,6 @@ void SObjectMixerEditorList::GenerateTreeView()
 
 		TreeViewRootObjects.AddUnique(TopLevelRow);
 	}
-	
-	ObjectsWithRowCreated.Empty();
 
 	TreeViewRootObjects.StableSort(SortByTypeThenName);
 
@@ -1145,37 +1212,47 @@ void SObjectMixerEditorList::FindVisibleObjectsAndRequestTreeRefresh()
 	TreeViewPtr->RequestTreeRefresh();
 }
 
-void SObjectMixerEditorList::SelectedTreeItemsToSelectedInLevelEditor() const
+void SObjectMixerEditorList::SelectedTreeItemsToSelectedInLevelEditor()
 {
-	if (GEditor)
+	if (GEditor && !bShouldPauseSyncSelection)
 	{
+		bShouldPauseSyncSelection = true;
 		TArray<AActor*> ActorsToSelect;
-		for (const TSharedPtr<FObjectMixerEditorListRow>& SelectedRow : TreeViewPtr->GetSelectedItems())
+		const TArray<FObjectMixerEditorListRowPtr>& SelectedTreeItems = TreeViewPtr->GetSelectedItems();
+		
+		if (SelectedTreeItems.Num() == 0)
 		{
-			if (SelectedRow->GetRowType() == FObjectMixerEditorListRow::MatchingObject ||
-				SelectedRow->GetRowType() == FObjectMixerEditorListRow::ContainerObject)
+			GEditor->SelectNone(true, true, true);
+		}
+		else
+		{
+			for (const TSharedPtr<FObjectMixerEditorListRow>& SelectedRow : SelectedTreeItems)
 			{
-				AActor* Actor = Cast<AActor>(SelectedRow->GetObject());
-
-				if (!Actor)
+				if (SelectedRow->GetRowType() == FObjectMixerEditorListRow::MatchingObject ||
+					SelectedRow->GetRowType() == FObjectMixerEditorListRow::ContainerObject)
 				{
-					Actor = SelectedRow->GetObject()->GetTypedOuter<AActor>();
-				}
+					AActor* Actor = Cast<AActor>(SelectedRow->GetObject());
 
-				if (Actor)
-				{
-					ActorsToSelect.Add(Actor);
+					if (!Actor)
+					{
+						Actor = SelectedRow->GetObject()->GetTypedOuter<AActor>();
+					}
+
+					if (Actor)
+					{
+						ActorsToSelect.Add(Actor);
+					}
 				}
 			}
-		}
-
-		if (ActorsToSelect.Num())
-		{
-			GEditor->SelectNone(true, true, true);
 
-			for (AActor* Actor : ActorsToSelect)
+			if (ActorsToSelect.Num())
 			{
-				GEditor->SelectActor( Actor, true, true, true );
+				GEditor->SelectNone(true, true, true);
+
+				for (AActor* Actor : ActorsToSelect)
+				{
+					GEditor->SelectActor( Actor, true, true, true );
+				}
 			}
 		}
 	}
diff --git a/Engine/Plugins/Editor/ObjectMixer/ObjectMixer/Source/ObjectMixer/Private/Views/List/SObjectMixerEditorListRow.cpp b/Engine/Plugins/Editor/ObjectMixer/ObjectMixer/Source/ObjectMixer/Private/Views/List/SObjectMixerEditorListRow.cpp
--- a/Engine/Plugins/Editor/ObjectMixer/ObjectMixer/Source/ObjectMixer/Private/Views/List/SObjectMixerEditorListRow.cpp
+++ b/Engine/Plugins/Editor/ObjectMixer/ObjectMixer/Source/ObjectMixer/Private/Views/List/SObjectMixerEditorListRow.cpp
@@ -25,9 +25,7 @@
 
 #define LOCTEXT_NAMESPACE "ObjectMixerEditor"
 
-const FText InsertFormatText = LOCTEXT("InsertAboveFormatText", "Insert {0} {1} {2}");
-const FText AboveText = LOCTEXT("AboveListItem", "above");
-const FText BelowText = LOCTEXT("BelowListItem", "below");
+const FText DropFormatText = LOCTEXT("DropFormatText", "{0} {1} {2} {3}");
 const FText MultiDragFormatText = LOCTEXT("MultiDragFormatText", "{0} Items");
 
 void SObjectMixerEditorListRow::Construct(
@@ -140,10 +138,8 @@ void SObjectMixerEditorListRow::HandleDragLeave(const FDragDropEvent& DragDropEv
 	}
 }
 
-TOptional<EItemDropZone> SObjectMixerEditorListRow::HandleCanAcceptDrop(const FDragDropEvent& DragDropEvent,
-                                                                             EItemDropZone DropZone,
-                                                                             FObjectMixerEditorListRowPtr
-                                                                             TargetItem)
+TOptional<EItemDropZone> SObjectMixerEditorListRow::HandleCanAcceptDrop(
+	const FDragDropEvent& DragDropEvent, EItemDropZone DropZone, FObjectMixerEditorListRowPtr TargetItem)
 {
 	TSharedPtr<FObjectMixerListRowDragDropOp> Operation =
 		DragDropEvent.GetOperationAs<FObjectMixerListRowDragDropOp>();
@@ -153,34 +149,50 @@ TOptional<EItemDropZone> SObjectMixerEditorListRow::HandleCanAcceptDrop(const FD
 		return TOptional<EItemDropZone>();
 	}
 
-	Operation->SetToolTip(
-		LOCTEXT("SortByCustomOrderDrgDropWarning", "Sort by custom order (\"#\") to drag & drop"),
-		FAppStyle::Get().GetBrush("Graph.ConnectorFeedback.Error")
-	);
-
-	const bool bIsDropDenied = false;
-
-	FString BoolAsString = bIsDropDenied ? "true" : "false";
+	const UObject* DropOnObject = TargetItem->GetObject();
+	const bool bIsDroppingOnFolderRow = TargetItem->GetRowType() == FObjectMixerEditorListRow::Folder;
+
+	const bool bIsDropDenied =
+		(!DropOnObject && !bIsDroppingOnFolderRow) ||
+		(DropOnObject && DropOnObject->IsA(UActorComponent::StaticClass())) ||
+		(
+			Operation->DraggedItems.Num() == 1 &&
+			Operation->DraggedItems[0]->GetObject() &&
+			Operation->DraggedItems[0]->GetObject()->IsA(UActorComponent::StaticClass())
+		) ||
+		(
+			Operation->DraggedItems.Num() == 1 &&
+			Operation->DraggedItems[0]->GetRowType() == FObjectMixerEditorListRow::Folder
+		)
+	;
 
 	if (bIsDropDenied)
 	{
-		Operation->ResetToDefaultToolTip();
+		Operation->SetToolTip(
+			LOCTEXT("ObjectMixerDragDropWarning", "Drop an actor row onto another actor row or folder to set attach parent or folder.\nDrop any row onto a collection button to assign a collection to the row."),
+			FAppStyle::Get().GetBrush("Graph.ConnectorFeedback.Error")
+		);
 
 		return TOptional<EItemDropZone>();
 	}
 
-	FText ItemNameText = FText::FromString("Light Label");
+	FText ItemNameText = FText::GetEmpty();
 
-	if (Operation->DraggedItems.Num() > 1)
+	if (Operation->DraggedItems.Num() == 1)
+	{
+		ItemNameText = Operation->DraggedItems[0]->GetDisplayName();
+	}
+	else
 	{
 		ItemNameText = FText::Format(MultiDragFormatText, FText::AsNumber(Operation->DraggedItems.Num()));
 	}
 
 	const FText DropPermittedText =
-		FText::Format(InsertFormatText,
-		              ItemNameText,
-		              DropZone == EItemDropZone::BelowItem ? BelowText : AboveText,
-		              ItemNameText
+		FText::Format(DropFormatText,
+			bIsDroppingOnFolderRow ? LOCTEXT("DragDropMoveToFolderPrefix", "Move") : LOCTEXT("DragDropSetAttachParentPrefix", "Set"),
+			ItemNameText,
+			bIsDroppingOnFolderRow ? LOCTEXT("DragDropMoveToFolderMidfix", "into") : LOCTEXT("DragDropSetAttachParentMidfix", "AttachParent as"),
+			TargetItem->GetDisplayName()
 		);
 
 	Operation->SetToolTip(
@@ -189,7 +201,7 @@ TOptional<EItemDropZone> SObjectMixerEditorListRow::HandleCanAcceptDrop(const FD
 	);
 
 	// We have no behaviour yet for dropping one item onto another, so we'll treat it like we dropped it above
-	return DropZone == EItemDropZone::OntoItem ? EItemDropZone::AboveItem : DropZone;
+	return EItemDropZone::OntoItem;
 }
 
 FReply SObjectMixerEditorListRow::HandleAcceptDrop(const FDragDropEvent& DragDropEvent, EItemDropZone DropZone,
@@ -203,39 +215,27 @@ FReply SObjectMixerEditorListRow::HandleAcceptDrop(const FDragDropEvent& DragDro
 		return FReply::Unhandled();
 	}
 
-	const TSharedPtr<SObjectMixerEditorList> ListView = Item.Pin()->GetListViewPtr().Pin();
-
-	TArray<FObjectMixerEditorListRowPtr> DraggedItems = Operation->DraggedItems;
-
-	TArray<FObjectMixerEditorListRowPtr> AllTreeItemsCopy = ListView->GetTreeViewItems();
-
-	for (const FObjectMixerEditorListRowPtr& DraggedItem : DraggedItems)
-	{
-		if (!DraggedItem.IsValid() || !AllTreeItemsCopy.Contains(DraggedItem))
-		{
-			continue;
-		}
-
-		AllTreeItemsCopy.Remove(DraggedItem);
-	}
+	UObject* DropOnObject = TargetItem->GetObject();
+	const bool bIsDroppingOnFolderRow = TargetItem->GetRowType() == FObjectMixerEditorListRow::Folder && TargetItem->GetFolderPath() != NAME_None;
+	FAttachmentTransformRules Rules(EAttachmentRule::KeepWorld, false);
 
-	const int32 TargetIndex = AllTreeItemsCopy.IndexOfByKey(TargetItem);
+	FScopedTransaction DragDropTransaction(LOCTEXT("ObjectMixerDragDropTransaction","Object Mixer Drag & Drop"));
 
-	if (TargetIndex > -1)
+	for (const FObjectMixerEditorListRowPtr& DraggedItem : Operation->DraggedItems)
 	{
-		for (int32 ItemIndex = DraggedItems.Num() - 1; ItemIndex >= 0; ItemIndex--)
+		if (AActor* ObjectAsActor = Cast<AActor>(DraggedItem->GetObject()))
 		{
-			const FObjectMixerEditorListRowPtr& DraggedItem = DraggedItems[ItemIndex];
-
-			if (!DraggedItem.IsValid() || AllTreeItemsCopy.Contains(DraggedItem))
+			if (bIsDroppingOnFolderRow)
 			{
-				continue;
+				ObjectAsActor->Modify();
+				ObjectAsActor->SetFolderPath(TargetItem->GetFolderPath());
+			}
+			else if (AActor* DropOnObjectAsActor = Cast<AActor>(DropOnObject))
+			{
+				ObjectAsActor->Modify();
+				ObjectAsActor->AttachToActor(DropOnObjectAsActor, Rules);
 			}
-
-			AllTreeItemsCopy.Insert(DraggedItem, DropZone == EItemDropZone::AboveItem ? TargetIndex : TargetIndex + 1);
 		}
-
-		ListView->SetTreeViewItems(AllTreeItemsCopy);
 	}
 
 	return FReply::Handled();
@@ -389,7 +389,7 @@ TSharedPtr<SWidget> SObjectMixerEditorListRow::GenerateCells(
 				.HAlign(HAlign_Left)
 				.VAlign(VAlign_Center)
 				.Visibility(EVisibility::SelfHitTestInvisible)
-				.Padding(0.f)
+				.Padding(FMargin(2,0,0,0))
 				[
 					SNew(SImage)
 					.ColorAndOpacity(this, &SObjectMixerEditorListRow::GetVisibilityIconForegroundColor)
@@ -429,18 +429,20 @@ TSharedPtr<SWidget> SObjectMixerEditorListRow::GenerateCells(
 		}
 		
 		return SNew(SBox)
-				.HAlign(HAlign_Left)
+				.HAlign(HAlign_Center)
 				.VAlign(VAlign_Center)
 				.Visibility(EVisibility::SelfHitTestInvisible)
 				.Padding(0.f)
 				[
 					SNew(SImage)
 					.ColorAndOpacity(this, &SObjectMixerEditorListRow::GetSoloIconForegroundColor)
-					.Image(FAppStyle::Get().GetBrush("MediaAsset.AssetActions.Solo.Small"))
+					.Image(FObjectMixerEditorStyle::Get().GetBrush("ObjectMixer.Solo"))
 					.OnMouseButtonDown_Lambda(
-						[PinnedItem] (const FGeometry& MyGeometry, const FPointerEvent& Event)
+						[this] (const FGeometry& MyGeometry, const FPointerEvent& Event)
 						{
-							check (PinnedItem);
+							check (Item.IsValid());
+
+							const FObjectMixerEditorListRowPtr PinnedItem = Item.Pin();
 							
 							if (const TSharedPtr<SObjectMixerEditorList> PinnedListView = PinnedItem->GetListViewPtr().Pin();
 								PinnedListView->GetTreeViewItemCount() > 0)
@@ -451,15 +453,13 @@ TSharedPtr<SWidget> SObjectMixerEditorListRow::GenerateCells(
 									TreeItem->SetObjectVisibility(bIsRowSolo, true);
 								}
 
-								if (bIsRowSolo)
-								{
-									PinnedItem->ClearSoloRow();
-								}
-								else
+								if (!bIsRowSolo)
 								{
 									PinnedItem->SetObjectVisibility(true, true);
-									PinnedItem->SetThisAsSoloRow();
+									PinnedItem->ClearSoloRows();
 								}
+								
+								PinnedItem->SetRowSoloState(!bIsRowSolo);
 
 								return FReply::Handled();
 							}
@@ -526,7 +526,15 @@ void SObjectMixerEditorListRow::OnPropertyChanged(const FProperty* Property, voi
 					
 					for (const TSharedPtr<FObjectMixerEditorListRow>& SelectedRow : PinnedItem->GetSelectedTreeViewItems())
 					{
-						if (UObject* SelectedRowObject = SelectedRow->GetObject())
+						UObject* SelectedRowObject = SelectedRow->GetObject();
+						const int32 SelectedRowHybridIndex = SelectedRow->GetHybridRowIndex();
+
+						if (SelectedRowHybridIndex != -1 && SelectedRowHybridIndex < SelectedRow->GetChildCount())
+						{
+							SelectedRowObject = SelectedRow->GetChildRows()[SelectedRowHybridIndex]->GetObject();
+						}
+						
+						if (SelectedRowObject)
 						{
 							Property->SetValue_InContainer(SelectedRowObject, ValuePtr);
 						}
diff --git a/Engine/Plugins/Editor/ObjectMixer/ObjectMixer/Source/ObjectMixer/Private/Views/MainPanel/ObjectMixerEditorMainPanel.cpp b/Engine/Plugins/Editor/ObjectMixer/ObjectMixer/Source/ObjectMixer/Private/Views/MainPanel/ObjectMixerEditorMainPanel.cpp
--- a/Engine/Plugins/Editor/ObjectMixer/ObjectMixer/Source/ObjectMixer/Private/Views/MainPanel/ObjectMixerEditorMainPanel.cpp
+++ b/Engine/Plugins/Editor/ObjectMixer/ObjectMixer/Source/ObjectMixer/Private/Views/MainPanel/ObjectMixerEditorMainPanel.cpp
@@ -44,6 +44,14 @@ void FObjectMixerEditorMainPanel::RefreshList() const
 	}
 }
 
+void FObjectMixerEditorMainPanel::RequestSyncEditorSelectionToListSelection() const
+{
+	if (EditorListModel.IsValid())
+	{
+		EditorListModel->RequestSyncEditorSelectionToListSelection();
+	}
+}
+
 void FObjectMixerEditorMainPanel::RebuildCollectionSelector()
 {
 	MainPanelWidget->RebuildCollectionSelector();
diff --git a/Engine/Plugins/Editor/ObjectMixer/ObjectMixer/Source/ObjectMixer/Public/ObjectMixerEditorModule.h b/Engine/Plugins/Editor/ObjectMixer/ObjectMixer/Source/ObjectMixer/Public/ObjectMixerEditorModule.h
--- a/Engine/Plugins/Editor/ObjectMixer/ObjectMixer/Source/ObjectMixer/Public/ObjectMixerEditorModule.h
+++ b/Engine/Plugins/Editor/ObjectMixer/ObjectMixer/Source/ObjectMixer/Public/ObjectMixerEditorModule.h
@@ -38,6 +38,8 @@ public:
 	 * Useful for when the list state has gone stale but the variable count has not changed.
 	 */
 	virtual void RefreshList() const;
+	
+	void RequestSyncEditorSelectionToListSelection();
 
 	void RegisterMenuGroup();
 	void UnregisterMenuGroup();
diff --git a/Engine/Plugins/Editor/ObjectMixer/ObjectMixer/Source/ObjectMixer/Public/Views/List/ObjectMixerEditorList.h b/Engine/Plugins/Editor/ObjectMixer/ObjectMixer/Source/ObjectMixer/Public/Views/List/ObjectMixerEditorList.h
--- a/Engine/Plugins/Editor/ObjectMixer/ObjectMixer/Source/ObjectMixer/Public/Views/List/ObjectMixerEditorList.h
+++ b/Engine/Plugins/Editor/ObjectMixer/ObjectMixer/Source/ObjectMixer/Public/Views/List/ObjectMixerEditorList.h
@@ -36,17 +36,20 @@ public:
 	 */
 	void RefreshList() const;
 
+	void RequestSyncEditorSelectionToListSelection() const;
+
 	void ExecuteListViewSearchOnAllRows(const FString& SearchString, const bool bShouldRefreshAfterward = true);
 
 	void EvaluateIfRowsPassFilters(const bool bShouldRefreshAfterward = true) const;
 
 	TWeakPtr<FObjectMixerEditorMainPanel> GetMainPanelModel();
 
-	TWeakPtr<FObjectMixerEditorListRow> GetSoloRow();
+	TSet<TWeakPtr<FObjectMixerEditorListRow>> GetSoloRows();
 
-	void SetSoloRow(TSharedRef<FObjectMixerEditorListRow> InRow);
+	void AddSoloRow(TSharedRef<FObjectMixerEditorListRow> InRow);
+	void RemoveSoloRow(TSharedRef<FObjectMixerEditorListRow> InRow);
 
-	void ClearSoloRow();
+	void ClearSoloRows();
 
 private:
 
diff --git a/Engine/Plugins/Editor/ObjectMixer/ObjectMixer/Source/ObjectMixer/Public/Views/List/ObjectMixerEditorListRow.h b/Engine/Plugins/Editor/ObjectMixer/ObjectMixer/Source/ObjectMixer/Public/Views/List/ObjectMixerEditorListRow.h
--- a/Engine/Plugins/Editor/ObjectMixer/ObjectMixer/Source/ObjectMixer/Public/Views/List/ObjectMixerEditorListRow.h
+++ b/Engine/Plugins/Editor/ObjectMixer/ObjectMixer/Source/ObjectMixer/Public/Views/List/ObjectMixerEditorListRow.h
@@ -71,6 +71,15 @@ struct OBJECTMIXEREDITOR_API FObjectMixerEditorListRow final : TSharedFromThis<F
 	, DisplayNameOverride(InDisplayNameOverride)
 	{}
 
+	FObjectMixerEditorListRow(
+		const FName InFolderPath, const EObjectMixerEditorListRowType InRowType, 
+		const TSharedRef<SObjectMixerEditorList>& InListView, const FText& InDisplayNameOverride = FText::GetEmpty())
+	: FolderPath(InFolderPath)
+	, RowType(InRowType)
+	, ListViewPtr(InListView)
+	, DisplayNameOverride(InDisplayNameOverride)
+	{}
+
 	[[nodiscard]] UObject* GetObject() const
 	{
 		if (ObjectRef.IsValid())
@@ -81,6 +90,11 @@ struct OBJECTMIXEREDITOR_API FObjectMixerEditorListRow final : TSharedFromThis<F
 		return nullptr;
 	}
 
+	[[nodiscard]] FName GetFolderPath() const
+	{
+		return FolderPath;
+	}
+
 	UObjectMixerObjectFilter* GetObjectFilter() const;
 
 	bool IsObjectRefInSelectedCollections() const;
@@ -161,16 +175,17 @@ struct OBJECTMIXEREDITOR_API FObjectMixerEditorListRow final : TSharedFromThis<F
 	bool GetObjectVisibility();
 	void SetObjectVisibility(const bool bNewIsVisible, const bool bIsRecursive = false);
 
-	bool IsThisRowSolo() const;
-	void SetThisAsSoloRow();
+	bool IsThisRowSolo();
+	void SetRowSoloState(const bool bNewSolo);
 
-	void ClearSoloRow();
+	void ClearSoloRows() const;
 
 private:
 
 	FObjectMixerEditorListRowPtr GetAsShared();
 
 	TWeakObjectPtr<UObject> ObjectRef;
+	FName FolderPath = NAME_None;
 	EObjectMixerEditorListRowType RowType = MatchingObject;
 	TArray<FObjectMixerEditorListRowPtr> ChildRows;
 	
diff --git a/Engine/Plugins/Editor/ObjectMixer/ObjectMixer/Source/ObjectMixer/Public/Views/List/SObjectMixerEditorList.h b/Engine/Plugins/Editor/ObjectMixer/ObjectMixer/Source/ObjectMixer/Public/Views/List/SObjectMixerEditorList.h
--- a/Engine/Plugins/Editor/ObjectMixer/ObjectMixer/Source/ObjectMixer/Public/Views/List/SObjectMixerEditorList.h
+++ b/Engine/Plugins/Editor/ObjectMixer/ObjectMixer/Source/ObjectMixer/Public/Views/List/SObjectMixerEditorList.h
@@ -100,6 +100,12 @@ public:
 
 	[[nodiscard]] TArray<FObjectMixerEditorListRowPtr> GetSelectedTreeViewItems() const;
 	int32 GetSelectedTreeViewItemCount() const;
+	
+	void RequestSyncEditorSelectionToListSelection()
+	{
+		bIsEditorToListSelectionSyncRequested = true;
+	}
+	void SyncEditorSelectionToListSelection();
 
 	void SetSelectedTreeViewItemActorsEditorVisible(const bool bNewIsVisible, const bool bIsRecursive = false);
 
@@ -113,11 +119,12 @@ public:
 		return TreeViewRootObjects.Num();
 	}
 	
-	TWeakPtr<FObjectMixerEditorListRow> GetSoloRow();
+	TSet<TWeakPtr<FObjectMixerEditorListRow>> GetSoloRows();
 
-	void SetSoloRow(TSharedRef<FObjectMixerEditorListRow> InRow);
+	void AddSoloRow(TSharedRef<FObjectMixerEditorListRow> InRow);
+	void RemoveSoloRow(TSharedRef<FObjectMixerEditorListRow> InRow);
 
-	void ClearSoloRow();
+	void ClearSoloRows();
 
 	FString GetSearchStringFromSearchInputField() const;
 	void ExecuteListViewSearchOnAllRows(const FString& SearchString, const bool bShouldRefreshAfterward = true);
@@ -182,7 +189,7 @@ protected:
 	TSharedPtr<SHeaderRow> HeaderRow;
 	TSharedRef<SWidget> GenerateHeaderRowContextMenu() const;
 
-	bool bShouldRebuild = false;
+	bool bIsRebuildRequested = false;
 	
 	/**
 	 * Regenerate the list items and refresh the list. Call when adding or removing items.
@@ -217,7 +224,11 @@ protected:
 	void FindVisibleTreeViewObjects();
 	void FindVisibleObjectsAndRequestTreeRefresh();
 
-	void SelectedTreeItemsToSelectedInLevelEditor() const;
+	void SelectedTreeItemsToSelectedInLevelEditor();
+	
+	/** For two-way selection sync, we need to pause selection sync under certain circumstances to prevent infinite loops. */
+	bool bShouldPauseSyncSelection = false;
+	bool bIsEditorToListSelectionSyncRequested = false;
 	
 	void OnGetRowChildren(FObjectMixerEditorListRowPtr Row, TArray<FObjectMixerEditorListRowPtr>& OutChildren) const;
 	void OnRowChildExpansionChange(FObjectMixerEditorListRowPtr Row, const bool bIsExpanded, const bool bIsRecursive = false) const;
@@ -226,7 +237,17 @@ protected:
 
 	TSharedPtr<STreeView<FObjectMixerEditorListRowPtr>> TreeViewPtr;
 
-	TMap<FString, bool> TreeItemExpansionStateCache;
+	struct FTreeItemStateCache
+	{
+		uint32 UniqueId = -1;
+		FString RowName = "";
+		bool bIsExpanded = false;
+		bool bIsSelected = false;
+	};
+
+	TArray<FTreeItemStateCache> TreeItemStateCache;
+	
+	TMap<UObject*, FObjectMixerEditorListRowPtr> ObjectsToRowsCreated;
 
 	/** All Tree view objects */
 	TArray<FObjectMixerEditorListRowPtr> TreeViewRootObjects;
diff --git a/Engine/Plugins/Editor/ObjectMixer/ObjectMixer/Source/ObjectMixer/Public/Views/MainPanel/ObjectMixerEditorMainPanel.h b/Engine/Plugins/Editor/ObjectMixer/ObjectMixer/Source/ObjectMixer/Public/Views/MainPanel/ObjectMixerEditorMainPanel.h
--- a/Engine/Plugins/Editor/ObjectMixer/ObjectMixer/Source/ObjectMixer/Public/Views/MainPanel/ObjectMixerEditorMainPanel.h
+++ b/Engine/Plugins/Editor/ObjectMixer/ObjectMixer/Source/ObjectMixer/Public/Views/MainPanel/ObjectMixerEditorMainPanel.h
@@ -39,6 +39,8 @@ public:
 	 */
 	void RefreshList() const;
 
+	void RequestSyncEditorSelectionToListSelection() const;
+
 	TWeakPtr<FObjectMixerEditorList> GetEditorListModel() const
 	{
 		return EditorListModel;
@@ -103,27 +105,35 @@ public:
 	const TArray<TSharedRef<IObjectMixerEditorListFilter>>& GetShowFilters() const;
 
 	/**
-	 * Get the row that has solo visibility. All other rows should be set to temporarily invisible in editor.
+	 * Get the rows that have solo visibility. All other rows should be set to temporarily invisible in editor.
+	 */
+	TSet<TWeakPtr<FObjectMixerEditorListRow>> GetSoloRows()
+	{
+		return SoloRows;
+	}
+
+	/**
+	 * Add a row that has solo visibility. This does not set temporary editor invisibility for other rows.
 	 */
-	TWeakPtr<FObjectMixerEditorListRow> GetSoloRow()
+	void AddSoloRow(TSharedRef<FObjectMixerEditorListRow> InRow)
 	{
-		return SoloRow;
+		SoloRows.Add(InRow);
 	}
 
 	/**
-	 * Set the row that has solo visibility. This does not set temporary editor invisibility for other rows.
+	 * Remove a row that does not have solo visibility. This does not set temporary editor invisibility for other rows.
 	 */
-	void SetSoloRow(TSharedRef<FObjectMixerEditorListRow> InRow)
+	void RemoveSoloRow(TSharedRef<FObjectMixerEditorListRow> InRow)
 	{
-		SoloRow = InRow;
+		SoloRows.Remove(InRow);
 	}
 
 	/**
-	 * Clear the row that has solo visibility. This does not remove temporary editor invisibility for other rows.
+	 * Clear the rows that have solo visibility. This does not remove temporary editor invisibility for other rows.
 	 */
-	void ClearSoloRow()
+	void ClearSoloRows()
 	{
-		SoloRow = nullptr;
+		SoloRows.Empty();
 	}
 
 	TSubclassOf<UObjectMixerObjectFilter> GetObjectFilterClass() const
@@ -186,7 +196,7 @@ private:
 	 */
 	EObjectMixerTreeViewMode TreeViewMode = EObjectMixerTreeViewMode::Folders;
 
-	TWeakPtr<FObjectMixerEditorListRow> SoloRow = nullptr;
+	TSet<TWeakPtr<FObjectMixerEditorListRow>> SoloRows = {};
 
 	FName ModuleName = NAME_None;
 
