Object Mixer:

-Fix: Add transaction for synced selections

diff --git a/Engine/Plugins/Editor/ObjectMixer/ObjectMixer/Source/ObjectMixer/Private/Views/List/SObjectMixerEditorList.cpp b/Engine/Plugins/Editor/ObjectMixer/ObjectMixer/Source/ObjectMixer/Private/Views/List/SObjectMixerEditorList.cpp
--- a/Engine/Plugins/Editor/ObjectMixer/ObjectMixer/Source/ObjectMixer/Private/Views/List/SObjectMixerEditorList.cpp
+++ b/Engine/Plugins/Editor/ObjectMixer/ObjectMixer/Source/ObjectMixer/Private/Views/List/SObjectMixerEditorList.cpp
@@ -23,6 +23,7 @@
 #include "DragAndDrop/AssetDragDropOp.h"
 #include "Editor/UnrealEdEngine.h"
 #include "PlacementMode/Public/IPlacementModeModule.h"
+#include "ScopedTransaction.h"
 #include "Styling/StyleColors.h"
 #include "Views/Widgets/ObjectMixerEditorListMenuContext.h"
 #include "Widgets/Layout/SScrollBox.h"
@@ -1247,10 +1248,13 @@ void SObjectMixerEditorList::SelectedTreeItemsToSelectedInLevelEditor()
 
 			if (ActorsToSelect.Num())
 			{
+				FScopedTransaction SyncSelection(LOCTEXT("SyncObjectMixerSelectionToScene","Sync Object Mixer Selection To Scene"));
+				
 				GEditor->SelectNone(true, true, true);
 
 				for (AActor* Actor : ActorsToSelect)
 				{
+					Actor->Modify();
 					GEditor->SelectActor( Actor, true, true, true );
 				}
 			}
