Object Mixer:

-Fix: Show Solo icon when hovering over rows when solo is inactive

diff --git a/Engine/Plugins/Editor/ObjectMixer/ObjectMixer/Source/ObjectMixer/Private/Views/List/SObjectMixerEditorListRow.cpp b/Engine/Plugins/Editor/ObjectMixer/ObjectMixer/Source/ObjectMixer/Private/Views/List/SObjectMixerEditorListRow.cpp
--- a/Engine/Plugins/Editor/ObjectMixer/ObjectMixer/Source/ObjectMixer/Private/Views/List/SObjectMixerEditorListRow.cpp
+++ b/Engine/Plugins/Editor/ObjectMixer/ObjectMixer/Source/ObjectMixer/Private/Views/List/SObjectMixerEditorListRow.cpp
@@ -264,7 +264,7 @@ FSlateColor SObjectMixerEditorListRow::GetVisibilityIconForegroundColor() const
 	}
 	else if (bIsHovered && !bIsSelected)
 	{
-		return FAppStyle::Get().GetSlateColor("Colors.ForegroundHover");
+		return FStyleColors::ForegroundHover;
 	}
 
 	return FSlateColor::UseForeground();
@@ -274,12 +274,19 @@ FSlateColor SObjectMixerEditorListRow::GetSoloIconForegroundColor() const
 {
 	check(Item.IsValid());
 
-	if (Item.Pin()->IsThisRowSolo())
+	const bool bIsSelected = Item.Pin()->GetIsSelected();
+
+	// make the foreground brush transparent if it is not selected, hovered or solo
+	if (!Item.Pin()->IsThisRowSolo() && !bIsHovered && !bIsSelected)
 	{
-		return bIsHovered ? FAppStyle::Get().GetSlateColor("Colors.ForegroundHover") : FSlateColor::UseForeground();
+		return FLinearColor::Transparent;
+	}
+	else if (bIsHovered && !bIsSelected)
+	{
+		return FStyleColors::ForegroundHover;
 	}
 
-	return FLinearColor::Transparent;
+	return FSlateColor::UseForeground();
 }
 
 const FSlateBrush* SObjectMixerEditorListRow::GetVisibilityBrush() const
