Object Mixer:

-Fix: Object Mixer/Light mixer context options should appear before general options in right-click context menu

diff --git a/Engine/Plugins/Editor/ObjectMixer/ObjectMixer/Source/ObjectMixer/Private/Views/Widgets/ObjectMixerEditorListMenuContext.cpp b/Engine/Plugins/Editor/ObjectMixer/ObjectMixer/Source/ObjectMixer/Private/Views/Widgets/ObjectMixerEditorListMenuContext.cpp
--- a/Engine/Plugins/Editor/ObjectMixer/ObjectMixer/Source/ObjectMixer/Private/Views/Widgets/ObjectMixerEditorListMenuContext.cpp
+++ b/Engine/Plugins/Editor/ObjectMixer/ObjectMixer/Source/ObjectMixer/Private/Views/Widgets/ObjectMixerEditorListMenuContext.cpp
@@ -118,8 +118,8 @@ void UObjectMixerEditorListMenuContext::RegisterObjectMixerContextMenu()
 					return;
 				}
 	
-				FToolMenuSection& Section = InMenu->FindOrAddSection("Collections");
-				Section.Label = LOCTEXT("CollectionsSectionName", "Collections");
+				FToolMenuSection& Section = InMenu->FindOrAddSection("ObjectMixerCollections");
+				Section.Label = LOCTEXT("ObjectMixerCollectionsSectionName", "Mixer Collections");
 				
 				FObjectMixerEditorListMenuContextData& ContextData = Context->Data;
 			
@@ -183,8 +183,9 @@ void UObjectMixerEditorListMenuContext::RegisterObjectMixerContextMenu()
 						}
 					)
 				);
-			}
-		));
+			}),
+			FToolMenuInsert(NAME_None,EToolMenuInsertType::First)
+		);
 	}
 }
 
