Object Mixer:

-Fix: Enabled columns are lost after restart

diff --git a/Engine/Plugins/Editor/ObjectMixer/ObjectMixer/Source/ObjectMixer/Private/ObjectMixerEditorSerializedData.cpp b/Engine/Plugins/Editor/ObjectMixer/ObjectMixer/Source/ObjectMixer/Private/ObjectMixerEditorSerializedData.cpp
--- a/Engine/Plugins/Editor/ObjectMixer/ObjectMixer/Source/ObjectMixer/Private/ObjectMixerEditorSerializedData.cpp
+++ b/Engine/Plugins/Editor/ObjectMixer/ObjectMixer/Source/ObjectMixer/Private/ObjectMixerEditorSerializedData.cpp
@@ -9,10 +9,10 @@
 
 #define LOCTEXT_NAMESPACE "ObjectMixerEditor"
 
-FObjectMixerSerializationData* UObjectMixerEditorSerializedData::FindSerializationDataByFilterClassName(const FName& FilterClassName)
+FObjectMixerSerializationDataPerFilter* UObjectMixerEditorSerializedData::FindSerializationDataByFilterClassName(const FName& FilterClassName)
 {
-	if (FObjectMixerSerializationData* Match = Algo::FindByPredicate(SerializedData,
-		[FilterClassName](const FObjectMixerSerializationData& Comparator)
+	if (FObjectMixerSerializationDataPerFilter* Match = Algo::FindByPredicate(SerializedDataPerFilter,
+		[FilterClassName](const FObjectMixerSerializationDataPerFilter& Comparator)
 		{
 			return Comparator.FilterClassName.IsEqual(FilterClassName);
 		}))
@@ -29,27 +29,33 @@ void UObjectMixerEditorSerializedData::AddObjectsToCollection(const FName& Filte
 
 	Modify();
 	
-	FObjectMixerSerializationData* Data = FindSerializationDataByFilterClassName(FilterClassName);
+	FObjectMixerSerializationDataPerFilter* Data = FindSerializationDataByFilterClassName(FilterClassName);
 
 	if (!Data)
 	{
-		SerializedData.Add({FilterClassName});
+		SerializedDataPerFilter.Add({FilterClassName});
 		Data = FindSerializationDataByFilterClassName(FilterClassName);
 	}
 	
 	if (Data)
 	{
-		if (FObjectMixerCollectionObjectSet* Match = Algo::FindByPredicate(Data->SerializedCollection,
+		TSet<FObjectMixerCollectionObjectData> NewObjectsData;
+		for (const FSoftObjectPath& Object : ObjectsToAdd)
+		{
+			NewObjectsData.Add({Object});
+		}
+		
+		if (FObjectMixerCollectionObjectSet* Match = Algo::FindByPredicate(Data->SerializedCollections,
 			[CollectionName](const FObjectMixerCollectionObjectSet& Comparator)
 			{
 				return Comparator.CollectionName.IsEqual(CollectionName);
 			}))
 		{
-			(*Match).CollectionObjects.Append(ObjectsToAdd);
+			Match->CollectionObjects.Append(NewObjectsData);
 		}
 		else
 		{
-			Data->SerializedCollection.Add({CollectionName, ObjectsToAdd});
+			Data->SerializedCollections.Add({CollectionName, NewObjectsData});
 		}
 
 		SaveConfig();
@@ -58,13 +64,13 @@ void UObjectMixerEditorSerializedData::AddObjectsToCollection(const FName& Filte
 
 void UObjectMixerEditorSerializedData::RemoveObjectsFromCollection(const FName& FilterClassName, const FName& CollectionName, const TSet<FSoftObjectPath>& ObjectsToRemove)
 {
-	if (FObjectMixerSerializationData* Data = FindSerializationDataByFilterClassName(FilterClassName))
+	if (FObjectMixerSerializationDataPerFilter* Data = FindSerializationDataByFilterClassName(FilterClassName))
 	{
 		FScopedTransaction RemoveCollectionTransaction(LOCTEXT("RemoveCollectionTransaction","Remove Collection"));
 	
 		Modify();
 		
-		Data->SerializedCollection.SetNum(Algo::StableRemoveIf(Data->SerializedCollection,
+		Data->SerializedCollections.SetNum(Algo::StableRemoveIf(Data->SerializedCollections,
 			[CollectionName](const FObjectMixerCollectionObjectSet& Comparator)
 			{
 				return Comparator.CollectionName.IsEqual(CollectionName);
@@ -76,13 +82,13 @@ void UObjectMixerEditorSerializedData::RemoveObjectsFromCollection(const FName&
 
 void UObjectMixerEditorSerializedData::RemoveCollection(const FName& FilterClassName, const FName& CollectionName)
 {	
-	if (FObjectMixerSerializationData* Data = FindSerializationDataByFilterClassName(FilterClassName))
+	if (FObjectMixerSerializationDataPerFilter* Data = FindSerializationDataByFilterClassName(FilterClassName))
 	{
 		FScopedTransaction RemoveCollectionTransaction(LOCTEXT("RemoveCollectionTransaction","Remove Collection"));
 	
 		Modify();
 		
-		Data->SerializedCollection.SetNum(Algo::StableRemoveIf(Data->SerializedCollection,
+		Data->SerializedCollections.SetNum(Algo::StableRemoveIf(Data->SerializedCollections,
 			[CollectionName](const FObjectMixerCollectionObjectSet& Comparator)
 			{
 				return Comparator.CollectionName.IsEqual(CollectionName);
@@ -95,11 +101,11 @@ void UObjectMixerEditorSerializedData::RemoveCollection(const FName& FilterClass
 void UObjectMixerEditorSerializedData::ReorderCollection(const FName& FilterClassName,
 	const FName& CollectionToMoveName, const FName& CollectionInsertBeforeName)
 {
-	if (FObjectMixerSerializationData* Data = FindSerializationDataByFilterClassName(FilterClassName))
+	if (FObjectMixerSerializationDataPerFilter* Data = FindSerializationDataByFilterClassName(FilterClassName))
 	{
 		FObjectMixerCollectionObjectSet MatchCopy;
 		bool bFoundMatch = false;
-		if (const FObjectMixerCollectionObjectSet* Match = Algo::FindByPredicate(Data->SerializedCollection,
+		if (const FObjectMixerCollectionObjectSet* Match = Algo::FindByPredicate(Data->SerializedCollections,
 			[CollectionToMoveName](const FObjectMixerCollectionObjectSet& Comparator)
 			{
 				return Comparator.CollectionName.IsEqual(CollectionToMoveName);
@@ -119,17 +125,17 @@ void UObjectMixerEditorSerializedData::ReorderCollection(const FName& FilterClas
 
 			if (CollectionInsertBeforeName == "All") // Move MatchCopy to the end
 			{
-				Data->SerializedCollection.Add(MatchCopy);
+				Data->SerializedCollections.Add(MatchCopy);
 			}
 			else
 			{
-				const int32 IndexOfCollectionInsertBefore = Algo::IndexOfByPredicate(Data->SerializedCollection,
+				const int32 IndexOfCollectionInsertBefore = Algo::IndexOfByPredicate(Data->SerializedCollections,
 					[CollectionInsertBeforeName](const FObjectMixerCollectionObjectSet& Comparator)
 					{
 						return Comparator.CollectionName.IsEqual(CollectionInsertBeforeName);
 					});
 
-				Data->SerializedCollection.Insert(MatchCopy, IndexOfCollectionInsertBefore);
+				Data->SerializedCollections.Insert(MatchCopy, IndexOfCollectionInsertBefore);
 			}
 		
 			SaveConfig();
@@ -139,15 +145,19 @@ void UObjectMixerEditorSerializedData::ReorderCollection(const FName& FilterClas
 
 bool UObjectMixerEditorSerializedData::IsObjectInCollection(const FName& FilterClassName, const FName& CollectionName, const FSoftObjectPath& InObject)
 {
-	if (FObjectMixerSerializationData* Data = FindSerializationDataByFilterClassName(FilterClassName))
+	if (FObjectMixerSerializationDataPerFilter* Data = FindSerializationDataByFilterClassName(FilterClassName))
 	{
-		if (const FObjectMixerCollectionObjectSet* Match = Algo::FindByPredicate(Data->SerializedCollection,
+		if (const FObjectMixerCollectionObjectSet* Match = Algo::FindByPredicate(Data->SerializedCollections,
 			[CollectionName](const FObjectMixerCollectionObjectSet& Comparator)
 			{
 				return Comparator.CollectionName.IsEqual(CollectionName);
 			}))
 		{
-			return (*Match).CollectionObjects.Contains(InObject);
+			return Algo::FindByPredicate((*Match).CollectionObjects,
+				[InObject](const FObjectMixerCollectionObjectData& Comparator)
+				{
+					return Comparator.ObjectPath == InObject;
+				}) != nullptr;
 		}
 	}
 
@@ -156,11 +166,11 @@ bool UObjectMixerEditorSerializedData::IsObjectInCollection(const FName& FilterC
 
 TSet<FName> UObjectMixerEditorSerializedData::GetCollectionsForObject(const FName& FilterClassName, const FSoftObjectPath& InObject)
 {
-	if (FObjectMixerSerializationData* Data = FindSerializationDataByFilterClassName(FilterClassName))
+	if (FObjectMixerSerializationDataPerFilter* Data = FindSerializationDataByFilterClassName(FilterClassName))
 	{
 		TSet<FName> CollectionsWithObject;
 
-		for (const FObjectMixerCollectionObjectSet& CollectionObjectSet : Data->SerializedCollection)
+		for (const FObjectMixerCollectionObjectSet& CollectionObjectSet : Data->SerializedCollections)
 		{
 			if (IsObjectInCollection(FilterClassName, CollectionObjectSet.CollectionName, InObject))
 			{
@@ -176,11 +186,11 @@ TSet<FName> UObjectMixerEditorSerializedData::GetCollectionsForObject(const FNam
 
 TArray<FName> UObjectMixerEditorSerializedData::GetAllCollectionNames(const FName& FilterClassName)
 {
-	if (FObjectMixerSerializationData* Data = FindSerializationDataByFilterClassName(FilterClassName))
+	if (FObjectMixerSerializationDataPerFilter* Data = FindSerializationDataByFilterClassName(FilterClassName))
 	{
 		TArray<FName> Collections;
 
-		for (const FObjectMixerCollectionObjectSet& CollectionObjectSet : Data->SerializedCollection)
+		for (const FObjectMixerCollectionObjectSet& CollectionObjectSet : Data->SerializedCollections)
 		{
 			Collections.Add(CollectionObjectSet.CollectionName);
 		}
@@ -191,4 +201,68 @@ TArray<FName> UObjectMixerEditorSerializedData::GetAllCollectionNames(const FNam
 	return {};
 }
 
+void UObjectMixerEditorSerializedData::SetShouldShowColumn(const FName& FilterClassName, const FName& ColumnName,
+	const bool bNewShouldShowColumn)
+{
+	FObjectMixerSerializationDataPerFilter* Data = FindSerializationDataByFilterClassName(FilterClassName);
+
+	if (!Data)
+	{
+		SerializedDataPerFilter.Add({FilterClassName});
+		Data = FindSerializationDataByFilterClassName(FilterClassName);
+	}
+	
+	if (Data)
+	{
+		if (FObjectMixerColumnData* Match = Algo::FindByPredicate(Data->SerializedColumnData,
+			[ColumnName](const FObjectMixerColumnData& Comparator)
+			{
+				return Comparator.ColumnName.IsEqual(ColumnName);
+			}))
+		{
+			(*Match).bShouldBeEnabled = bNewShouldShowColumn;
+		}
+		else
+		{
+			Data->SerializedColumnData.Add({ColumnName, bNewShouldShowColumn});
+		}
+
+		SaveConfig();
+	}
+}
+
+bool UObjectMixerEditorSerializedData::IsColumnDataSerialized(const FName& FilterClassName, const FName& ColumnName)
+{
+	if (FObjectMixerSerializationDataPerFilter* Data = FindSerializationDataByFilterClassName(FilterClassName))
+	{
+		if (const FObjectMixerColumnData* Match = Algo::FindByPredicate(Data->SerializedColumnData,
+			[ColumnName](const FObjectMixerColumnData& Comparator)
+			{
+				return Comparator.ColumnName.IsEqual(ColumnName);
+			}))
+		{
+			return true;
+		}
+	}
+
+	return false;
+}
+
+bool UObjectMixerEditorSerializedData::ShouldShowColumn(const FName& FilterClassName, const FName& ColumnName)
+{
+	if (FObjectMixerSerializationDataPerFilter* Data = FindSerializationDataByFilterClassName(FilterClassName))
+	{
+		if (const FObjectMixerColumnData* Match = Algo::FindByPredicate(Data->SerializedColumnData,
+			[ColumnName](const FObjectMixerColumnData& Comparator)
+			{
+				return Comparator.ColumnName.IsEqual(ColumnName);
+			}))
+		{
+			return (*Match).bShouldBeEnabled;
+		}
+	}
+
+	return false;
+}
+
 #undef LOCTEXT_NAMESPACE
diff --git a/Engine/Plugins/Editor/ObjectMixer/ObjectMixer/Source/ObjectMixer/Private/Views/List/SObjectMixerEditorList.cpp b/Engine/Plugins/Editor/ObjectMixer/ObjectMixer/Source/ObjectMixer/Private/Views/List/SObjectMixerEditorList.cpp
--- a/Engine/Plugins/Editor/ObjectMixer/ObjectMixer/Source/ObjectMixer/Private/Views/List/SObjectMixerEditorList.cpp
+++ b/Engine/Plugins/Editor/ObjectMixer/ObjectMixer/Source/ObjectMixer/Private/Views/List/SObjectMixerEditorList.cpp
@@ -18,6 +18,7 @@
 #include "LevelEditor.h"
 #include "LevelEditorMenuContext.h"
 #include "LevelEditorViewport.h"
+#include "ObjectMixerEditorSerializedData.h"
 #include "Selection.h"
 #include "Framework/MultiBox/MultiBoxBuilder.h"
 #include "DragAndDrop/AssetDragDropOp.h"
@@ -280,6 +281,10 @@ void SObjectMixerEditorList::ClearSoloRows()
 
 FString SObjectMixerEditorList::GetSearchStringFromSearchInputField() const
 {
+	check(ListModelPtr.IsValid());
+	check(ListModelPtr.Pin()->GetMainPanelModel().IsValid());
+	check(HeaderRow);
+	
 	return ListModelPtr.Pin()->GetMainPanelModel().Pin()->GetSearchStringFromSearchInputField();
 }
 
@@ -525,9 +530,22 @@ TSharedRef<SWidget> SObjectMixerEditorList::GenerateHeaderRowContextMenu() const
 			FUIAction(
 				FExecuteAction::CreateLambda([this, PropertyName]()
 				{
+					check(ListModelPtr.IsValid());
+					check(ListModelPtr.Pin()->GetMainPanelModel().IsValid());
 					check(HeaderRow);
+
+					const bool bNewColumnEnabled = !HeaderRow->IsColumnVisible(PropertyName);
 					
-					HeaderRow->SetShowGeneratedColumn(PropertyName, !HeaderRow->IsColumnVisible(PropertyName));
+					HeaderRow->SetShowGeneratedColumn(PropertyName, bNewColumnEnabled);
+
+					if (UObjectMixerEditorSerializedData* SerializedData = GetMutableDefault<UObjectMixerEditorSerializedData>())
+					{
+						if (const TSubclassOf<UObjectMixerObjectFilter> Filter = ListModelPtr.Pin()->GetMainPanelModel().Pin()->GetObjectFilterClass())
+						{
+							const FName FilterName = Filter->GetFName();
+							SerializedData->SetShouldShowColumn(FilterName, PropertyName, bNewColumnEnabled);
+						}
+					}
 				}),
 				FCanExecuteAction::CreateLambda([bCanSelectColumn](){return bCanSelectColumn;}),
 				FIsActionChecked::CreateLambda([this, PropertyName]()
@@ -650,15 +668,8 @@ void SObjectMixerEditorList::AddBuiltinColumnsToHeaderRow()
 TSharedPtr<SHeaderRow> SObjectMixerEditorList::GenerateHeaderRow()
 {
 	check(ListModelPtr.IsValid());
-	check(ListModelPtr.Pin()->GetMainPanelModel().Pin());
+	check(ListModelPtr.Pin()->GetMainPanelModel().IsValid());
 	check(HeaderRow);
-	
-	TMap<FName, bool> LastVisibleColumns;
-	
-	for (const SHeaderRow::FColumn& Column : HeaderRow->GetColumns())
-	{
-		LastVisibleColumns.Add(Column.ColumnId, Column.bIsVisible);
-	}
 
 	HeaderRow->ClearColumns();
 	ListViewColumns.Empty(ListViewColumns.Num());
@@ -782,11 +793,22 @@ TSharedPtr<SHeaderRow> SObjectMixerEditorList::GenerateHeaderRow()
 			}
 		
 			HeaderRow->AddColumn(Column);
+
+			// Figure out if we should show a certain column
+			// First check filter rules
 			bool bShouldShowColumn = ColumnsToShowByDefaultCache.Contains(ColumnInfo.PropertyName);
 
-			if (const bool* Match = LastVisibleColumns.Find(ColumnInfo.PropertyName))
+			// Then load visible columns from SerializedData
+			if (UObjectMixerEditorSerializedData* SerializedData = GetMutableDefault<UObjectMixerEditorSerializedData>())
 			{
-				bShouldShowColumn = *Match;
+				if (const TSubclassOf<UObjectMixerObjectFilter> Filter = ListModelPtr.Pin()->GetMainPanelModel().Pin()->GetObjectFilterClass())
+				{
+					const FName FilterName = Filter->GetFName();
+					if (SerializedData->IsColumnDataSerialized(FilterName, ColumnInfo.PropertyName))
+					{
+						bShouldShowColumn = SerializedData->ShouldShowColumn(FilterName, ColumnInfo.PropertyName);
+					}
+				}
 			}
 		
 			HeaderRow->SetShowGeneratedColumn(ColumnInfo.PropertyName, bShouldShowColumn);
@@ -929,7 +951,7 @@ void SObjectMixerEditorList::RestoreTreeState(const bool bFlushCache)
 void SObjectMixerEditorList::BuildPerformanceCacheAndGenerateHeaderIfNeeded()
 {
 	check(ListModelPtr.IsValid());
-	check(ListModelPtr.Pin()->GetMainPanelModel().Pin());
+	check(ListModelPtr.Pin()->GetMainPanelModel().IsValid());
 	
 	// If any of the following overrides change, we need to regenerate the header row. Otherwise skip regeneration for performance reasons.
 	// GetObjectClassesToFilter, GetColumnsToShowByDefault, GetColumnsToExclude,
diff --git a/Engine/Plugins/Editor/ObjectMixer/ObjectMixer/Source/ObjectMixer/Public/ObjectMixerEditorProjectSettings.h b/Engine/Plugins/Editor/ObjectMixer/ObjectMixer/Source/ObjectMixer/Public/ObjectMixerEditorProjectSettings.h
--- a/Engine/Plugins/Editor/ObjectMixer/ObjectMixer/Source/ObjectMixer/Public/ObjectMixerEditorProjectSettings.h
+++ b/Engine/Plugins/Editor/ObjectMixer/ObjectMixer/Source/ObjectMixer/Public/ObjectMixerEditorProjectSettings.h
@@ -4,7 +4,7 @@
 
 #include "ObjectMixerEditorProjectSettings.generated.h"
 
-UCLASS(config = ObjectMixer, defaultconfig)
+UCLASS(config = ObjectMixer)
 class OBJECTMIXEREDITOR_API UObjectMixerEditorProjectSettings : public UObject
 {
 	GENERATED_BODY()
@@ -26,4 +26,4 @@ public:
 	 */
 	UPROPERTY(Config, EditAnywhere, Category="Object Mixer")
 	bool bExpandTreeViewItemsByDefault = true;
-};
\ No newline at end of file
+};
diff --git a/Engine/Plugins/Editor/ObjectMixer/ObjectMixer/Source/ObjectMixer/Public/ObjectMixerEditorSerializedData.h b/Engine/Plugins/Editor/ObjectMixer/ObjectMixer/Source/ObjectMixer/Public/ObjectMixerEditorSerializedData.h
--- a/Engine/Plugins/Editor/ObjectMixer/ObjectMixer/Source/ObjectMixer/Public/ObjectMixerEditorSerializedData.h
+++ b/Engine/Plugins/Editor/ObjectMixer/ObjectMixer/Source/ObjectMixer/Public/ObjectMixerEditorSerializedData.h
@@ -4,6 +4,26 @@
 
 #include "ObjectMixerEditorSerializedData.generated.h"
 
+USTRUCT()
+struct FObjectMixerCollectionObjectData
+{
+	GENERATED_BODY()
+
+	UPROPERTY()
+	FSoftObjectPath ObjectPath = {};
+
+	bool operator==(const FObjectMixerCollectionObjectData& Other) const
+	{
+		return ObjectPath == Other.ObjectPath;
+	}
+
+	friend uint32 GetTypeHash (const FObjectMixerCollectionObjectData& Other)
+	{
+		return GetTypeHash(Other.ObjectPath);
+	}
+
+};
+
 USTRUCT()
 struct FObjectMixerCollectionObjectSet
 {
@@ -13,7 +33,7 @@ struct FObjectMixerCollectionObjectSet
 	FName CollectionName = NAME_None;
 
 	UPROPERTY()
-	TSet<FSoftObjectPath> CollectionObjects = {};
+	TSet<FObjectMixerCollectionObjectData> CollectionObjects = {};
 
 	bool operator==(const FObjectMixerCollectionObjectSet& Other) const
 	{
@@ -27,7 +47,29 @@ struct FObjectMixerCollectionObjectSet
 };
 
 USTRUCT()
-struct FObjectMixerSerializationData
+struct FObjectMixerColumnData
+{
+	GENERATED_BODY()
+	
+	UPROPERTY()
+	FName ColumnName = NAME_None;
+
+	UPROPERTY()
+	bool bShouldBeEnabled = false;
+
+	bool operator==(const FObjectMixerColumnData& Other) const
+	{
+		return ColumnName.IsEqual(Other.ColumnName);
+	}
+
+	friend uint32 GetTypeHash (const FObjectMixerColumnData& Other)
+	{
+		return GetTypeHash(Other.ColumnName);
+	}
+};
+
+USTRUCT()
+struct FObjectMixerSerializationDataPerFilter
 {
 	GENERATED_BODY()
 	
@@ -35,14 +77,17 @@ struct FObjectMixerSerializationData
 	FName FilterClassName = NAME_None;
 
 	UPROPERTY()
-	TArray<FObjectMixerCollectionObjectSet> SerializedCollection = {};
+	TArray<FObjectMixerCollectionObjectSet> SerializedCollections = {};
+
+	UPROPERTY()
+	TSet<FObjectMixerColumnData> SerializedColumnData = {};
 
-	bool operator==(const FObjectMixerSerializationData& Other) const
+	bool operator==(const FObjectMixerSerializationDataPerFilter& Other) const
 	{
 		return FilterClassName.IsEqual(Other.FilterClassName);
 	}
 
-	friend uint32 GetTypeHash (const FObjectMixerSerializationData& Other)
+	friend uint32 GetTypeHash (const FObjectMixerSerializationDataPerFilter& Other)
 	{
 		return GetTypeHash(Other.FilterClassName);
 	}
@@ -61,9 +106,9 @@ public:
 	}
 	
 	UPROPERTY(Config)
-	TSet<FObjectMixerSerializationData> SerializedData;
+	TSet<FObjectMixerSerializationDataPerFilter> SerializedDataPerFilter;
 
-	FObjectMixerSerializationData* FindSerializationDataByFilterClassName(const FName& FilterClassName);
+	FObjectMixerSerializationDataPerFilter* FindSerializationDataByFilterClassName(const FName& FilterClassName);
 
 	void AddObjectsToCollection(const FName& FilterClassName, const FName& CollectionName, const TSet<FSoftObjectPath>& ObjectsToAdd);
 
@@ -78,4 +123,10 @@ public:
 	TSet<FName> GetCollectionsForObject(const FName& FilterClassName, const FSoftObjectPath& InObject);
 
 	TArray<FName> GetAllCollectionNames(const FName& FilterClassName);
+
+	void SetShouldShowColumn(const FName& FilterClassName, const FName& ColumnName, const bool bNewShouldShowColumn);
+
+	bool IsColumnDataSerialized(const FName& FilterClassName, const FName& ColumnName);
+	
+	bool ShouldShowColumn(const FName& FilterClassName, const FName& ColumnName);
 };
