Object Mixer:

-Highlight Search Text
-Add SHyperlinkWithTextHighlight widget

diff --git a/Engine/Plugins/Editor/ObjectMixer/ObjectMixer/Source/ObjectMixer/Private/Views/List/SObjectMixerEditorList.cpp b/Engine/Plugins/Editor/ObjectMixer/ObjectMixer/Source/ObjectMixer/Private/Views/List/SObjectMixerEditorList.cpp
--- a/Engine/Plugins/Editor/ObjectMixer/ObjectMixer/Source/ObjectMixer/Private/Views/List/SObjectMixerEditorList.cpp
+++ b/Engine/Plugins/Editor/ObjectMixer/ObjectMixer/Source/ObjectMixer/Private/Views/List/SObjectMixerEditorList.cpp
@@ -279,13 +279,22 @@ void SObjectMixerEditorList::ClearSoloRows()
 	GetListModelPtr().Pin()->ClearSoloRows();
 }
 
+FText SObjectMixerEditorList::GetSearchTextFromSearchInputField() const
+{
+	TSharedPtr<FObjectMixerEditorList> PinnedListModel = ListModelPtr.Pin();
+	check(PinnedListModel && PinnedListModel->GetMainPanelModel().IsValid());
+	check(HeaderRow);
+	
+	return PinnedListModel->GetMainPanelModel().Pin()->GetSearchTextFromSearchInputField();
+}
+
 FString SObjectMixerEditorList::GetSearchStringFromSearchInputField() const
 {
-	check(ListModelPtr.IsValid());
-	check(ListModelPtr.Pin()->GetMainPanelModel().IsValid());
+	TSharedPtr<FObjectMixerEditorList> PinnedListModel = ListModelPtr.Pin();
+	check(PinnedListModel && PinnedListModel->GetMainPanelModel().IsValid());
 	check(HeaderRow);
 	
-	return ListModelPtr.Pin()->GetMainPanelModel().Pin()->GetSearchStringFromSearchInputField();
+	return PinnedListModel->GetMainPanelModel().Pin()->GetSearchStringFromSearchInputField();
 }
 
 void SObjectMixerEditorList::ExecuteListViewSearchOnAllRows(
@@ -530,8 +539,8 @@ TSharedRef<SWidget> SObjectMixerEditorList::GenerateHeaderRowContextMenu() const
 			FUIAction(
 				FExecuteAction::CreateLambda([this, PropertyName]()
 				{
-					check(ListModelPtr.IsValid());
-					check(ListModelPtr.Pin()->GetMainPanelModel().IsValid());
+					TSharedPtr<FObjectMixerEditorList> PinnedListModel = ListModelPtr.Pin();
+					check(PinnedListModel && PinnedListModel->GetMainPanelModel().IsValid());
 					check(HeaderRow);
 
 					const bool bNewColumnEnabled = !HeaderRow->IsColumnVisible(PropertyName);
@@ -540,7 +549,7 @@ TSharedRef<SWidget> SObjectMixerEditorList::GenerateHeaderRowContextMenu() const
 
 					if (UObjectMixerEditorSerializedData* SerializedData = GetMutableDefault<UObjectMixerEditorSerializedData>())
 					{
-						if (const TSubclassOf<UObjectMixerObjectFilter> Filter = ListModelPtr.Pin()->GetMainPanelModel().Pin()->GetObjectFilterClass())
+						if (const TSubclassOf<UObjectMixerObjectFilter> Filter = PinnedListModel->GetMainPanelModel().Pin()->GetObjectFilterClass())
 						{
 							const FName FilterName = Filter->GetFName();
 							SerializedData->SetShouldShowColumn(FilterName, PropertyName, bNewColumnEnabled);
@@ -667,15 +676,15 @@ void SObjectMixerEditorList::AddBuiltinColumnsToHeaderRow()
 
 TSharedPtr<SHeaderRow> SObjectMixerEditorList::GenerateHeaderRow()
 {
-	check(ListModelPtr.IsValid());
-	check(ListModelPtr.Pin()->GetMainPanelModel().IsValid());
+	TSharedPtr<FObjectMixerEditorList> PinnedListModel = ListModelPtr.Pin();
+	check(PinnedListModel && PinnedListModel->GetMainPanelModel().IsValid());
 	check(HeaderRow);
 
 	HeaderRow->ClearColumns();
 	ListViewColumns.Empty(ListViewColumns.Num());
 	
 	TSet<UClass*> SpecifiedClasses =
-		ListModelPtr.Pin()->GetMainPanelModel().Pin()->GetObjectFilter()->GetParentAndChildClassesFromSpecifiedClasses(
+		PinnedListModel->GetMainPanelModel().Pin()->GetObjectFilter()->GetParentAndChildClassesFromSpecifiedClasses(
 			ObjectClassesToFilterCache, PropertyInheritanceInclusionOptionsCache);
 	
 	for (const UClass* Class : SpecifiedClasses)
@@ -801,7 +810,7 @@ TSharedPtr<SHeaderRow> SObjectMixerEditorList::GenerateHeaderRow()
 			// Then load visible columns from SerializedData
 			if (UObjectMixerEditorSerializedData* SerializedData = GetMutableDefault<UObjectMixerEditorSerializedData>())
 			{
-				if (const TSubclassOf<UObjectMixerObjectFilter> Filter = ListModelPtr.Pin()->GetMainPanelModel().Pin()->GetObjectFilterClass())
+				if (const TSubclassOf<UObjectMixerObjectFilter> Filter = PinnedListModel->GetMainPanelModel().Pin()->GetObjectFilterClass())
 				{
 					const FName FilterName = Filter->GetFName();
 					if (SerializedData->IsColumnDataSerialized(FilterName, ColumnInfo.PropertyName))
@@ -950,15 +959,15 @@ void SObjectMixerEditorList::RestoreTreeState(const bool bFlushCache)
 
 void SObjectMixerEditorList::BuildPerformanceCacheAndGenerateHeaderIfNeeded()
 {
-	check(ListModelPtr.IsValid());
-	check(ListModelPtr.Pin()->GetMainPanelModel().IsValid());
+	TSharedPtr<FObjectMixerEditorList> PinnedListModel = ListModelPtr.Pin();
+	check(PinnedListModel && PinnedListModel->GetMainPanelModel().IsValid());
 	
 	// If any of the following overrides change, we need to regenerate the header row. Otherwise skip regeneration for performance reasons.
 	// GetObjectClassesToFilter, GetColumnsToShowByDefault, GetColumnsToExclude,
 	// GetForceAddedColumns, GetObjectMixerPropertyInheritanceInclusionOptions, ShouldIncludeUnsupportedProperties
 	bool bNeedToGenerateHeaders = false;
 	
-	const TObjectPtr<UObjectMixerObjectFilter> SelectedFilter = ListModelPtr.Pin()->GetMainPanelModel().Pin()->GetObjectFilter();
+	const TObjectPtr<UObjectMixerObjectFilter> SelectedFilter = PinnedListModel->GetMainPanelModel().Pin()->GetObjectFilter();
 	if (!SelectedFilter)
 	{
 		UE_LOG(LogObjectMixerEditor, Display, TEXT("%hs: No classes defined in UObjectMixerObjectFilter class."), __FUNCTION__);
@@ -1157,7 +1166,8 @@ void CreateFolderRowsForTopLevelRow(
 
 void SObjectMixerEditorList::GenerateTreeView()
 {
-	check(ListModelPtr.IsValid());
+	TSharedPtr<FObjectMixerEditorList> PinnedListModel = ListModelPtr.Pin();
+	check(PinnedListModel);
 	
 	if (!ensure(TreeViewPtr.IsValid()))
 	{
@@ -1206,7 +1216,7 @@ void SObjectMixerEditorList::GenerateTreeView()
 
 	TreeViewRootObjects.StableSort(SortByTypeThenName);
 
-	if (TSharedPtr<FObjectMixerEditorMainPanel> MainPanel = ListModelPtr.Pin()->GetMainPanelModel().Pin())
+	if (TSharedPtr<FObjectMixerEditorMainPanel> MainPanel = PinnedListModel->GetMainPanelModel().Pin())
 	{
 		MainPanel->RebuildCollectionSelector();
 	}
diff --git a/Engine/Plugins/Editor/ObjectMixer/ObjectMixer/Source/ObjectMixer/Private/Views/List/SObjectMixerEditorListRow.cpp b/Engine/Plugins/Editor/ObjectMixer/ObjectMixer/Source/ObjectMixer/Private/Views/List/SObjectMixerEditorListRow.cpp
--- a/Engine/Plugins/Editor/ObjectMixer/ObjectMixer/Source/ObjectMixer/Private/Views/List/SObjectMixerEditorListRow.cpp
+++ b/Engine/Plugins/Editor/ObjectMixer/ObjectMixer/Source/ObjectMixer/Private/Views/List/SObjectMixerEditorListRow.cpp
@@ -6,8 +6,10 @@
 #include "ObjectMixerEditorStyle.h"
 #include "Views/List/ObjectMixerEditorList.h"
 #include "Views/List/SObjectMixerEditorList.h"
+#include "Views/Widgets/SHyperlinkWithTextHighlight.h"
 
 #include "Customizations/ColorStructCustomization.h"
+#include "Editor.h"
 #include "Engine/Blueprint.h"
 #include "GameFramework/Actor.h"
 #include "Input/DragAndDrop.h"
@@ -20,7 +22,6 @@
 #include "Styling/StyleColors.h"
 #include "Widgets/Layout/SBox.h"
 #include "Widgets/SBoxPanel.h"
-#include "Widgets/Layout/SScaleBox.h"
 #include "Widgets/Text/STextBlock.h"
 
 #define LOCTEXT_NAMESPACE "ObjectMixerEditor"
@@ -334,19 +335,19 @@ TSharedPtr<SWidget> SObjectMixerEditorListRow::GenerateCells(
 			if (UClass* ActorClass = Object->GetClass())
 			{
 				if (UBlueprint* AsBlueprint = UBlueprint::GetBlueprintFromClass(ActorClass))
-				{
+				{					
 					bNeedsStandardTextBlock = false;
-					
-					FEditorClassUtils::FSourceLinkParams SourceLinkParams;
-					SourceLinkParams.Object = Object;
-					SourceLinkParams.bUseDefaultFormat = false;
-					SourceLinkParams.bUseFormatIfNoLink = true;
-					SourceLinkParams.BlueprintFormat = &DisplayName;
-
+			
 					HBox->AddSlot()
 					.Padding(FMargin(10.0, 0, 0, 0))
 					[
-						FEditorClassUtils::GetSourceLink(ActorClass, SourceLinkParams)
+						SNew(SHyperlinkWithTextHighlight)
+						.Style(FAppStyle::Get(), "Common.GotoBlueprintHyperlink")
+						.Visibility(EVisibility::Visible)
+						.Text(DisplayName)
+						.ToolTipText(LOCTEXT("ClickToEditBlueprint", "Click to edit Blueprint"))
+						.OnNavigate(this, &SObjectMixerEditorListRow::OnClickBlueprintLink, AsBlueprint, Object.Get())
+						.HighlightText(this, &SObjectMixerEditorListRow::GetHighlightText)
 					];
 				}
 			}
@@ -368,10 +369,11 @@ TSharedPtr<SWidget> SObjectMixerEditorListRow::GenerateCells(
 			.Padding(FMargin(10.0, 0, 0, 0))
 			[
 				SNew(STextBlock)
-				.Visibility(EVisibility::Visible)
+				.Visibility(EVisibility::SelfHitTestInvisible)
 				.Justification(ETextJustify::Left)
 				.Text(DisplayName)
 				.ToolTipText(TooltipText)
+				.HighlightText(this, &SObjectMixerEditorListRow::GetHighlightText)
 			];
 		}
 		
@@ -526,7 +528,7 @@ TSharedPtr<SWidget> SObjectMixerEditorListRow::GenerateCells(
 	return nullptr;
 }
 
-void SObjectMixerEditorListRow::OnPropertyChanged(const FProperty* Property, void* ContainerWithChangedProperty)
+void SObjectMixerEditorListRow::OnPropertyChanged(const FProperty* Property, void* ContainerWithChangedProperty) const
 {
 	if (Property && ContainerWithChangedProperty)
 	{
@@ -559,4 +561,32 @@ void SObjectMixerEditorListRow::OnPropertyChanged(const FProperty* Property, voi
 	}
 }
 
+void SObjectMixerEditorListRow::OnClickBlueprintLink(UBlueprint* AsBlueprint, UObject* Object)
+{
+	if (AsBlueprint)
+	{
+		if (Object)
+		{
+			if (ensure(Object->GetClass()->ClassGeneratedBy == AsBlueprint))
+			{
+				AsBlueprint->SetObjectBeingDebugged(Object);
+			}
+		}
+		// Open the blueprint
+		GEditor->EditObject(AsBlueprint);
+	}
+}
+
+FText SObjectMixerEditorListRow::GetHighlightText() const
+{
+	check (Item.IsValid());
+
+	if (const TSharedPtr<SObjectMixerEditorList> PinnedListView = Item.Pin()->GetListViewPtr().Pin())
+	{
+		return PinnedListView->GetSearchTextFromSearchInputField();
+	}
+
+	return FText::GetEmpty();
+}
+
 #undef LOCTEXT_NAMESPACE
diff --git a/Engine/Plugins/Editor/ObjectMixer/ObjectMixer/Source/ObjectMixer/Private/Views/MainPanel/ObjectMixerEditorMainPanel.cpp b/Engine/Plugins/Editor/ObjectMixer/ObjectMixer/Source/ObjectMixer/Private/Views/MainPanel/ObjectMixerEditorMainPanel.cpp
--- a/Engine/Plugins/Editor/ObjectMixer/ObjectMixer/Source/ObjectMixer/Private/Views/MainPanel/ObjectMixerEditorMainPanel.cpp
+++ b/Engine/Plugins/Editor/ObjectMixer/ObjectMixer/Source/ObjectMixer/Private/Views/MainPanel/ObjectMixerEditorMainPanel.cpp
@@ -57,6 +57,11 @@ void FObjectMixerEditorMainPanel::RebuildCollectionSelector()
 	MainPanelWidget->RebuildCollectionSelector();
 }
 
+FText FObjectMixerEditorMainPanel::GetSearchTextFromSearchInputField() const
+{
+	return MainPanelWidget->GetSearchTextFromSearchInputField();
+}
+
 FString FObjectMixerEditorMainPanel::GetSearchStringFromSearchInputField() const
 {
 	return MainPanelWidget->GetSearchStringFromSearchInputField();
diff --git a/Engine/Plugins/Editor/ObjectMixer/ObjectMixer/Source/ObjectMixer/Private/Views/MainPanel/SObjectMixerEditorMainPanel.cpp b/Engine/Plugins/Editor/ObjectMixer/ObjectMixer/Source/ObjectMixer/Private/Views/MainPanel/SObjectMixerEditorMainPanel.cpp
--- a/Engine/Plugins/Editor/ObjectMixer/ObjectMixer/Source/ObjectMixer/Private/Views/MainPanel/SObjectMixerEditorMainPanel.cpp
+++ b/Engine/Plugins/Editor/ObjectMixer/ObjectMixer/Source/ObjectMixer/Private/Views/MainPanel/SObjectMixerEditorMainPanel.cpp
@@ -413,6 +413,13 @@ void SObjectMixerEditorMainPanel::OnSearchTextChanged(const FText& Text)
 	ExecuteListViewSearchOnAllRows(Text.ToString(), true);
 }
 
+FText SObjectMixerEditorMainPanel::GetSearchTextFromSearchInputField() const
+{
+	return ensureAlwaysMsgf(SearchBoxPtr.IsValid(),
+		TEXT("%hs: SearchBoxPtr is not valid. Check to make sure it was created."), __FUNCTION__)
+	? SearchBoxPtr->GetText() : FText::GetEmpty();
+}
+
 FString SObjectMixerEditorMainPanel::GetSearchStringFromSearchInputField() const
 {
 	return ensureAlwaysMsgf(SearchBoxPtr.IsValid(),
diff --git a/Engine/Plugins/Editor/ObjectMixer/ObjectMixer/Source/ObjectMixer/Public/Views/List/SObjectMixerEditorList.h b/Engine/Plugins/Editor/ObjectMixer/ObjectMixer/Source/ObjectMixer/Public/Views/List/SObjectMixerEditorList.h
--- a/Engine/Plugins/Editor/ObjectMixer/ObjectMixer/Source/ObjectMixer/Public/Views/List/SObjectMixerEditorList.h
+++ b/Engine/Plugins/Editor/ObjectMixer/ObjectMixer/Source/ObjectMixer/Public/Views/List/SObjectMixerEditorList.h
@@ -126,6 +126,7 @@ public:
 
 	void ClearSoloRows();
 
+	FText GetSearchTextFromSearchInputField() const;
 	FString GetSearchStringFromSearchInputField() const;
 	void ExecuteListViewSearchOnAllRows(const FString& SearchString, const bool bShouldRefreshAfterward = true);
 
diff --git a/Engine/Plugins/Editor/ObjectMixer/ObjectMixer/Source/ObjectMixer/Public/Views/List/SObjectMixerEditorListRow.h b/Engine/Plugins/Editor/ObjectMixer/ObjectMixer/Source/ObjectMixer/Public/Views/List/SObjectMixerEditorListRow.h
--- a/Engine/Plugins/Editor/ObjectMixer/ObjectMixer/Source/ObjectMixer/Public/Views/List/SObjectMixerEditorListRow.h
+++ b/Engine/Plugins/Editor/ObjectMixer/ObjectMixer/Source/ObjectMixer/Public/Views/List/SObjectMixerEditorListRow.h
@@ -49,7 +49,10 @@ private:
 
 	TSharedPtr<SWidget> GenerateCells(const FName& InColumnName, const TSharedPtr<FObjectMixerEditorListRow> PinnedItem);
 
-	void OnPropertyChanged(const FProperty* Property, void* ContainerWithChangedProperty);
+	void OnPropertyChanged(const FProperty* Property, void* ContainerWithChangedProperty) const;
+
+	void OnClickBlueprintLink(UBlueprint* AsBlueprint, UObject* Object = nullptr);
+	FText GetHighlightText() const;
 	
 	TWeakPtr<FObjectMixerEditorListRow> Item;
 	
diff --git a/Engine/Plugins/Editor/ObjectMixer/ObjectMixer/Source/ObjectMixer/Public/Views/MainPanel/ObjectMixerEditorMainPanel.h b/Engine/Plugins/Editor/ObjectMixer/ObjectMixer/Source/ObjectMixer/Public/Views/MainPanel/ObjectMixerEditorMainPanel.h
--- a/Engine/Plugins/Editor/ObjectMixer/ObjectMixer/Source/ObjectMixer/Public/Views/MainPanel/ObjectMixerEditorMainPanel.h
+++ b/Engine/Plugins/Editor/ObjectMixer/ObjectMixer/Source/ObjectMixer/Public/Views/MainPanel/ObjectMixerEditorMainPanel.h
@@ -48,6 +48,7 @@ public:
 
 	void RebuildCollectionSelector();
 
+	FText GetSearchTextFromSearchInputField() const;
 	FString GetSearchStringFromSearchInputField() const;
 
 	void OnClassSelectionChanged(UClass* InNewClass);
diff --git a/Engine/Plugins/Editor/ObjectMixer/ObjectMixer/Source/ObjectMixer/Public/Views/MainPanel/SObjectMixerEditorMainPanel.h b/Engine/Plugins/Editor/ObjectMixer/ObjectMixer/Source/ObjectMixer/Public/Views/MainPanel/SObjectMixerEditorMainPanel.h
--- a/Engine/Plugins/Editor/ObjectMixer/ObjectMixer/Source/ObjectMixer/Public/Views/MainPanel/SObjectMixerEditorMainPanel.h
+++ b/Engine/Plugins/Editor/ObjectMixer/ObjectMixer/Source/ObjectMixer/Public/Views/MainPanel/SObjectMixerEditorMainPanel.h
@@ -27,6 +27,7 @@ public:
 		return MainPanelModel;
 	}
 
+	FText GetSearchTextFromSearchInputField() const;
 	FString GetSearchStringFromSearchInputField() const;
 	void SetSearchStringInSearchInputField(const FString InSearchString) const;
 	void ExecuteListViewSearchOnAllRows(const FString& SearchString, const bool bShouldRefreshAfterward = true);
diff --git a/Engine/Plugins/Editor/ObjectMixer/ObjectMixer/Source/ObjectMixer/Public/Views/Widgets/SHyperlinkWithTextHighlight.h b/Engine/Plugins/Editor/ObjectMixer/ObjectMixer/Source/ObjectMixer/Public/Views/Widgets/SHyperlinkWithTextHighlight.h
new file mode 100644
--- /dev/null
+++ b/Engine/Plugins/Editor/ObjectMixer/ObjectMixer/Source/ObjectMixer/Public/Views/Widgets/SHyperlinkWithTextHighlight.h
@@ -0,0 +1,122 @@
+// Copyright Epic Games, Inc. All Rights Reserved.
+
+#pragma once
+
+#include "Misc/Attribute.h"
+#include "Input/CursorReply.h"
+#include "Input/Reply.h"
+#include "Layout/Margin.h"
+#include "Styling/AppStyle.h"
+#include "Styling/SlateColor.h"
+#include "Styling/SlateTypes.h"
+#include "Widgets/DeclarativeSyntaxSupport.h"
+#include "Widgets/Input/SButton.h"
+#include "Widgets/Text/STextBlock.h"
+
+#if WITH_ACCESSIBILITY
+#include "Widgets/Accessibility/SlateAccessibleWidgets.h"
+#endif
+
+/**
+ * A hyperlink widget is what you would expect from a browser hyperlink.
+ * When a hyperlink is clicked in invokes an OnNavigate() delegate.
+ */
+class SHyperlinkWithTextHighlight
+	: public SButton
+{
+public:
+
+	SLATE_BEGIN_ARGS(SHyperlinkWithTextHighlight)
+		: _Text()
+		, _Style(&FAppStyle::Get().GetWidgetStyle< FHyperlinkStyle >("Hyperlink"))
+		, _TextStyle(nullptr)
+		, _UnderlineStyle(nullptr)
+		, _Padding()
+		, _OnNavigate()
+		, _TextShapingMethod()
+		, _TextFlowDirection()
+		, _HighlightColor()
+		, _HighlightShape()
+		, _HighlightText()
+		{}
+
+		SLATE_ATTRIBUTE( FText, Text )
+		SLATE_STYLE_ARGUMENT( FHyperlinkStyle, Style )
+		SLATE_STYLE_ARGUMENT( FTextBlockStyle, TextStyle )
+		SLATE_STYLE_ARGUMENT( FButtonStyle, UnderlineStyle )
+		SLATE_ATTRIBUTE( FMargin, Padding )
+		SLATE_EVENT( FSimpleDelegate, OnNavigate )
+		SLATE_ARGUMENT( TOptional<ETextShapingMethod>, TextShapingMethod )
+		SLATE_ARGUMENT( TOptional<ETextFlowDirection>, TextFlowDirection )
+
+		/** The color used to highlight the specified text */
+		SLATE_ATTRIBUTE( FLinearColor, HighlightColor )
+
+		/** The brush used to highlight the specified text */
+		SLATE_ATTRIBUTE( const FSlateBrush*, HighlightShape )
+
+		/** Highlight this text in the text block */
+		SLATE_ATTRIBUTE( FText, HighlightText )
+		SLATE_END_ARGS()
+
+	/**
+	 * Construct the hyperlink widgets from a declaration
+	 * 
+	 * @param InArgs    Widget declaration from which to construct the hyperlink.
+	 */
+	void Construct( const FArguments& InArgs )
+	{
+		this->OnNavigate = InArgs._OnNavigate;
+
+		check (InArgs._Style);
+		const FButtonStyle* UnderlineStyle = InArgs._UnderlineStyle != nullptr ? InArgs._UnderlineStyle : &InArgs._Style->UnderlineStyle;
+		const FTextBlockStyle* TextStyle = InArgs._TextStyle != nullptr ? InArgs._TextStyle : &InArgs._Style->TextStyle;
+		TAttribute<FMargin> Padding = InArgs._Padding.IsSet() ? InArgs._Padding : InArgs._Style->Padding;
+
+		SButton::Construct(
+			SButton::FArguments()
+			.ContentPadding(Padding)
+			.ButtonStyle(UnderlineStyle)
+			.OnClicked(this, &SHyperlinkWithTextHighlight::Hyperlink_OnClicked)
+			.ForegroundColor(FSlateColor::UseForeground())
+			.TextShapingMethod(InArgs._TextShapingMethod)
+			.TextFlowDirection(InArgs._TextFlowDirection)
+			[
+				SNew(STextBlock)
+				.TextStyle(TextStyle)
+				.Text(InArgs._Text)
+				.HighlightColor(InArgs._HighlightColor)
+				.HighlightShape(InArgs._HighlightShape)
+				.HighlightText(InArgs._HighlightText)
+			]
+		);
+	}
+
+public:
+
+	// SWidget overrides
+
+	virtual FCursorReply OnCursorQuery( const FGeometry& MyGeometry, const FPointerEvent& CursorEvent ) const override
+	{
+		return FCursorReply::Cursor( EMouseCursor::Hand );
+	}
+#if WITH_ACCESSIBILITY
+	virtual TSharedRef<FSlateAccessibleWidget> CreateAccessibleWidget() override
+	{
+		return MakeShareable<FSlateAccessibleWidget>(new FSlateAccessibleHyperlink(SharedThis(this)));
+	}
+#endif
+
+protected:
+
+	/** Invoke the OnNavigate method */
+	FReply Hyperlink_OnClicked()
+	{
+		OnNavigate.ExecuteIfBound();
+
+		return FReply::Handled();
+	}
+
+	/** The delegate to invoke when someone clicks the hyperlink */
+	FSimpleDelegate OnNavigate;
+};
