Object Mixer:

-Fix: Spawnable lights stop displaying in Light Mixer once keyframed

diff --git a/Engine/Plugins/Editor/ObjectMixer/ObjectMixer/Source/ObjectMixer/Private/ObjectMixerEditorModule.cpp b/Engine/Plugins/Editor/ObjectMixer/ObjectMixer/Source/ObjectMixer/Private/ObjectMixerEditorModule.cpp
--- a/Engine/Plugins/Editor/ObjectMixer/ObjectMixer/Source/ObjectMixer/Private/ObjectMixerEditorModule.cpp
+++ b/Engine/Plugins/Editor/ObjectMixer/ObjectMixer/Source/ObjectMixer/Private/ObjectMixerEditorModule.cpp
@@ -71,6 +71,9 @@ void FObjectMixerEditorModule::Teardown()
 	if (GEditor)
 	{
 		GEditor->OnBlueprintCompiled().RemoveAll(this);
+		GEditor->OnLevelActorAdded().RemoveAll(this);
+		GEditor->OnLevelActorDeleted().RemoveAll(this);
+		GEditor->GetSelectedActors()->SelectionChangedEvent.RemoveAll(this);
 	}
 	
 	FEditorDelegates::MapChange.RemoveAll(this);
@@ -307,6 +310,10 @@ void FObjectMixerEditorModule::BindDelegates()
 	{
 		RequestRebuildList();
 	}));
+	DelegateHandles.Add(GEditor->OnLevelActorAdded().AddLambda([this] (AActor*)
+	{
+		RequestRebuildList();
+	}));
 	DelegateHandles.Add(GEditor->OnLevelActorDeleted().AddLambda([this] (AActor*)
 	{
 		RequestRebuildList();
