Object Mixer:

-Update BlueprintType UENUM comments to be more artist-friendly

diff --git a/Engine/Plugins/Editor/ObjectMixer/ObjectMixer/Source/ObjectMixer/Public/ObjectFilter/ObjectMixerEditorObjectFilter.h b/Engine/Plugins/Editor/ObjectMixer/ObjectMixer/Source/ObjectMixer/Public/ObjectFilter/ObjectMixerEditorObjectFilter.h
--- a/Engine/Plugins/Editor/ObjectMixer/ObjectMixer/Source/ObjectMixer/Public/ObjectFilter/ObjectMixerEditorObjectFilter.h
+++ b/Engine/Plugins/Editor/ObjectMixer/ObjectMixer/Source/ObjectMixer/Public/ObjectFilter/ObjectMixerEditorObjectFilter.h
@@ -10,22 +10,33 @@
 UENUM(BlueprintType)
 enum class EObjectMixerInheritanceInclusionOptions : uint8
 {
-	None, // Get only the properties in the specified classes without considering parent or child classes + Specified Class
-	IncludeOnlyImmediateParent, // Get properties from the class the specified class immediately derives from, but not their parents + Specified Class
-	IncludeOnlyImmediateChildren, // Get properties from child classes but not child classes of child classes + Specified Class
-	IncludeOnlyImmediateParentAndChildren, // IncludeOnlyImmediateParent + IncludeOnlyImmediateChildren + Specified Class
-	IncludeAllParents, // Go up the chain of Super classes to get all properties in the class' ancestry + Specified Class
-	IncludeAllChildren, // Get properties from all derived classes recursively + Specified Class
-	IncludeAllParentsAndChildren, // IncludeAllParents + IncludeAllChildren + Specified Class
-	IncludeAllParentsAndOnlyImmediateChildren, // IncludeAllParents + IncludeOnlyImmediateChildren + Specified Class
-	IncludeOnlyImmediateParentAndAllChildren // IncludeOnlyImmediateParent + IncludeAllChildren + Specified Class
+	// Get only the properties in the specified classes without considering parent or child classes
+	None,
+	// Get properties from the class that the specified classes immediately derive from, but not the parents' parents + Specified Classes
+	IncludeOnlyImmediateParent,
+	// Get properties from child classes but not child classes of child classes + Specified Classes
+	IncludeOnlyImmediateChildren,
+	// IncludeOnlyImmediateParent + IncludeOnlyImmediateChildren + Specified Classes
+	IncludeOnlyImmediateParentAndChildren,
+	// Go up the chain of parent classes to get all properties in the specified classes' ancestries + Specified Classes
+	IncludeAllParents,
+	// Get properties from all derived classes recursively + Specified Classes
+	IncludeAllChildren,
+	// IncludeAllParents + IncludeAllChildren + Specified Classes
+	IncludeAllParentsAndChildren,
+	// IncludeAllParents + IncludeOnlyImmediateChildren + Specified Classes
+	IncludeAllParentsAndOnlyImmediateChildren, 
+	// IncludeOnlyImmediateParent + IncludeAllChildren + Specified Classes
+	IncludeOnlyImmediateParentAndAllChildren 
 };
 
 UENUM(BlueprintType)
 enum class EObjectMixerTreeViewMode : uint8
 {
-	NoFolders, // Show all matching objects in a flat list without any hierarchy
-	Folders // Display objects in a folder hierarchy but without Object->SubObject hierarchy
+	// Show all matching objects without folders
+	NoFolders,
+	// Display objects in a hierarchy with folders
+	Folders 
 };
 
 /**
