Object Mixer:

-Fix: Blueprint filters for Object Mixer don't appear in Filter list until loaded

diff --git a/Engine/Plugins/Editor/ObjectMixer/ObjectMixer/Source/ObjectMixer/Private/Views/MainPanel/SObjectMixerEditorMainPanel.cpp b/Engine/Plugins/Editor/ObjectMixer/ObjectMixer/Source/ObjectMixer/Private/Views/MainPanel/SObjectMixerEditorMainPanel.cpp
--- a/Engine/Plugins/Editor/ObjectMixer/ObjectMixer/Source/ObjectMixer/Private/Views/MainPanel/SObjectMixerEditorMainPanel.cpp
+++ b/Engine/Plugins/Editor/ObjectMixer/ObjectMixer/Source/ObjectMixer/Private/Views/MainPanel/SObjectMixerEditorMainPanel.cpp
@@ -10,11 +10,12 @@
 #include "Views/Widgets/SCollectionSelectionButton.h"
 #include "Views/Widgets/SObjectMixerPlacementAssetMenuEntry.h"
 
+#include "AssetRegistry/AssetRegistryModule.h"
+#include "AssetRegistry/IAssetRegistry.h"
 #include "Framework/MultiBox/MultiBoxBuilder.h"
 #include "Kismet2/SClassPickerDialog.h"
 #include "PlacementMode/Public/IPlacementModeModule.h"
 #include "SPositiveActionButton.h"
-#include "Brushes/SlateRoundedBoxBrush.h"
 #include "Styling/StyleColors.h"
 #include "Views/List/ObjectMixerEditorListFilters/IObjectMixerEditorListFilter.h"
 #include "Views/List/ObjectMixerEditorListFilters/ObjectMixerEditorListFilter_Collection.h"
@@ -189,16 +190,53 @@ TSharedRef<SWidget> SObjectMixerEditorMainPanel::OnGenerateAddObjectButtonMenu()
 	;
 }
 
+bool IsBlueprintFilter(const FAssetData& BlueprintClassData)
+{
+	UClass* BlueprintFilterClass = UObjectMixerBlueprintObjectFilter::StaticClass();
+		
+	const FString NativeParentClassPath = BlueprintClassData.GetTagValueRef<FString>(FBlueprintTags::NativeParentClassPath);
+	const FSoftClassPath ClassPath(NativeParentClassPath);
+			
+	UClass* NativeParentClass = ClassPath.ResolveClass();
+	const bool bInheritsFromBlueprintFilter =
+		NativeParentClass // Class may have been removed, or renamed and not correctly redirected
+		&& (NativeParentClass == BlueprintFilterClass || NativeParentClass->IsChildOf(BlueprintFilterClass));
+
+	return bInheritsFromBlueprintFilter;
+}
+
 TSharedRef<SWidget> SObjectMixerEditorMainPanel::OnGenerateFilterClassMenu()
 {
 	FMenuBuilder MenuBuilder(true, nullptr);
 
+	// Get C++ Derivatives (and maybe Blueprint derivatives)
 	TArray<UClass*> DerivedClasses;
 	GetDerivedClasses(UObjectMixerObjectFilter::StaticClass(), DerivedClasses, true);
 
 	DerivedClasses.Remove(UObjectMixerObjectFilter::StaticClass());
 	DerivedClasses.Remove(UObjectMixerBlueprintObjectFilter::StaticClass());
 
+	// Get remaining Blueprint derivatives
+	FAssetRegistryModule& AssetRegistryModule = FModuleManager::LoadModuleChecked< FAssetRegistryModule >(FName("AssetRegistry"));
+	IAssetRegistry& AssetRegistry = AssetRegistryModule.Get();
+
+	TArray< FAssetData > Assets;
+	AssetRegistry.GetAssetsByClass(UBlueprint::StaticClass()->GetClassPathName(), Assets);	
+	for (const FAssetData& Asset : Assets)
+	{
+		if (IsBlueprintFilter(Asset))
+		{
+			if (UBlueprint* BlueprintAsset = Cast<UBlueprint>(Asset.GetAsset()))
+			{
+				UClass* LoadedClass = BlueprintAsset->GeneratedClass;
+				if (ensure(LoadedClass && BlueprintAsset->ParentClass))
+				{
+					DerivedClasses.AddUnique(LoadedClass);
+				}
+			}
+		}
+	}
+
 	DerivedClasses.Sort([](UClass& A, UClass& B)
 	{
 		return A.GetFName().LexicalLess(B.GetFName());
@@ -225,8 +263,8 @@ TSharedRef<SWidget> SObjectMixerEditorMainPanel::OnGenerateFilterClassMenu()
 					}
 					
 					MenuBuilder.AddMenuEntry(
-					   FText::FromName(DerivedClass->GetFName()),
-					   FText::GetEmpty(),
+					   FText::FromString(DerivedClass->GetName().EndsWith(TEXT("_C")) ? DerivedClass->GetName().LeftChop(2) : DerivedClass->GetName()),
+					   FText::FromString(DerivedClass->GetClassPathName().ToString()),
 					   FSlateIcon(),
 					   FUIAction(
 						   FExecuteAction::CreateSP(MainPanelModel.Pin().ToSharedRef(), &FObjectMixerEditorMainPanel::SetObjectFilterClass, DerivedClass),
diff --git a/Engine/Plugins/Editor/ObjectMixer/ObjectMixer/Source/ObjectMixer/Public/ObjectFilter/ObjectMixerEditorObjectFilter.h b/Engine/Plugins/Editor/ObjectMixer/ObjectMixer/Source/ObjectMixer/Public/ObjectFilter/ObjectMixerEditorObjectFilter.h
--- a/Engine/Plugins/Editor/ObjectMixer/ObjectMixer/Source/ObjectMixer/Public/ObjectFilter/ObjectMixerEditorObjectFilter.h
+++ b/Engine/Plugins/Editor/ObjectMixer/ObjectMixer/Source/ObjectMixer/Public/ObjectFilter/ObjectMixerEditorObjectFilter.h
@@ -43,7 +43,7 @@ enum class EObjectMixerTreeViewMode : uint8
  * Native class for filtering object types to Object Mixer.
  * Native C++ classes should inherit directly from this class.
  */
-UCLASS(Abstract, Blueprintable) 
+UCLASS(Abstract, BlueprintType) 
 class OBJECTMIXEREDITOR_API UObjectMixerObjectFilter : public UObject
 {
 	GENERATED_BODY()
