Object Mixer:

-Fix: Panel should scroll to selected

diff --git a/Engine/Plugins/Editor/ObjectMixer/ObjectMixer/Source/ObjectMixer/Private/Views/List/SObjectMixerEditorList.cpp b/Engine/Plugins/Editor/ObjectMixer/ObjectMixer/Source/ObjectMixer/Private/Views/List/SObjectMixerEditorList.cpp
--- a/Engine/Plugins/Editor/ObjectMixer/ObjectMixer/Source/ObjectMixer/Private/Views/List/SObjectMixerEditorList.cpp
+++ b/Engine/Plugins/Editor/ObjectMixer/ObjectMixer/Source/ObjectMixer/Private/Views/List/SObjectMixerEditorList.cpp
@@ -220,16 +220,25 @@ void SObjectMixerEditorList::SyncEditorSelectionToListSelection()
 		{
 			bShouldPauseSyncSelection = true;
 			TreeViewPtr->ClearSelection();
-		
+
+			FObjectMixerEditorListRowPtr* ListRowPtr = nullptr;
 			for (int32 SelectionItr = 0; SelectionItr < SelectionCount; SelectionItr++)
 			{
 				UObject* SelectedActor = SelectedActors->GetSelectedObject(SelectionItr);
 
-				if (const FObjectMixerEditorListRowPtr* Match = ObjectsToRowsCreated.Find(SelectedActor))
+				ListRowPtr = ObjectsToRowsCreated.Find(SelectedActor);
+
+				if (ListRowPtr)
 				{
-					TreeViewPtr->SetItemSelection(*Match, true);
+					TreeViewPtr->SetItemSelection(*ListRowPtr, true);
 				}
 			}
+			
+			// Scroll to the last selected item
+			if (ListRowPtr)
+			{
+				TreeViewPtr->RequestScrollIntoView(*ListRowPtr);
+			}
 		}
 	}
 }
