Object Mixer:

-Fix: off state for Solo icon

diff --git a/Engine/Plugins/Editor/ObjectMixer/ObjectMixer/Source/ObjectMixer/Private/ObjectMixerEditorStyle.cpp b/Engine/Plugins/Editor/ObjectMixer/ObjectMixer/Source/ObjectMixer/Private/ObjectMixerEditorStyle.cpp
--- a/Engine/Plugins/Editor/ObjectMixer/ObjectMixer/Source/ObjectMixer/Private/ObjectMixerEditorStyle.cpp
+++ b/Engine/Plugins/Editor/ObjectMixer/ObjectMixer/Source/ObjectMixer/Private/ObjectMixerEditorStyle.cpp
@@ -85,6 +85,7 @@ TSharedRef< FSlateStyleSet > FObjectMixerEditorStyle::Create()
 	Style->Set("ObjectMixer.ToolbarButton.Small", new IMAGE_BRUSH_SVG("Icons/ObjectMixer", Icon20x20));
 	
 	Style->Set("ObjectMixer.Solo", new IMAGE_BRUSH_SVG("Icons/Solo", Icon16x16));
+	Style->Set("ObjectMixer.SoloHoverOff", new IMAGE_BRUSH_SVG("Icons/DashSmall", Icon16x16));
 	
 	return Style;
 }
diff --git a/Engine/Plugins/Editor/ObjectMixer/ObjectMixer/Source/ObjectMixer/Private/Views/List/SObjectMixerEditorListRow.cpp b/Engine/Plugins/Editor/ObjectMixer/ObjectMixer/Source/ObjectMixer/Private/Views/List/SObjectMixerEditorListRow.cpp
--- a/Engine/Plugins/Editor/ObjectMixer/ObjectMixer/Source/ObjectMixer/Private/Views/List/SObjectMixerEditorListRow.cpp
+++ b/Engine/Plugins/Editor/ObjectMixer/ObjectMixer/Source/ObjectMixer/Private/Views/List/SObjectMixerEditorListRow.cpp
@@ -57,6 +57,9 @@ void SObjectMixerEditorListRow::Construct(
 	VisibleNotHoveredBrush = FAppStyle::Get().GetBrush(VisibleNotHoveredBrushName);
 	NotVisibleHoveredBrush = FAppStyle::Get().GetBrush(NotVisibleHoveredBrushName);
 	NotVisibleNotHoveredBrush = FAppStyle::Get().GetBrush(NotVisibleNotHoveredBrushName);
+	
+	SoloOnBrush = FObjectMixerEditorStyle::Get().GetBrush("ObjectMixer.Solo");
+	SoloOffHoveredBrush = FObjectMixerEditorStyle::Get().GetBrush("ObjectMixer.SoloHoverOff");
 }
 
 TSharedRef<SWidget> SObjectMixerEditorListRow::GenerateWidgetForColumn(const FName& InColumnName)
@@ -300,6 +303,19 @@ const FSlateBrush* SObjectMixerEditorListRow::GetVisibilityBrush() const
 	return bIsHovered ? NotVisibleHoveredBrush : NotVisibleNotHoveredBrush;
 }
 
+const FSlateBrush* SObjectMixerEditorListRow::GetSoloBrush() const
+{
+	check(Item.IsValid());
+
+	// make the foreground brush transparent if it is not selected, hovered or solo
+	if (Item.Pin()->IsThisRowSolo())
+	{
+		return SoloOnBrush;
+	}
+
+	return SoloOffHoveredBrush;
+}
+
 TSharedPtr<SWidget> SObjectMixerEditorListRow::GenerateCells(
 	const FName& InColumnName, const TSharedPtr<FObjectMixerEditorListRow> PinnedItem)
 {
@@ -445,7 +461,7 @@ TSharedPtr<SWidget> SObjectMixerEditorListRow::GenerateCells(
 				[
 					SNew(SImage)
 					.ColorAndOpacity(this, &SObjectMixerEditorListRow::GetSoloIconForegroundColor)
-					.Image(FObjectMixerEditorStyle::Get().GetBrush("ObjectMixer.Solo"))
+					.Image(this, &SObjectMixerEditorListRow::GetSoloBrush)
 					.OnMouseButtonDown_Lambda(
 						[this] (const FGeometry& MyGeometry, const FPointerEvent& Event)
 						{
diff --git a/Engine/Plugins/Editor/ObjectMixer/ObjectMixer/Source/ObjectMixer/Public/Views/List/SObjectMixerEditorListRow.h b/Engine/Plugins/Editor/ObjectMixer/ObjectMixer/Source/ObjectMixer/Public/Views/List/SObjectMixerEditorListRow.h
--- a/Engine/Plugins/Editor/ObjectMixer/ObjectMixer/Source/ObjectMixer/Public/Views/List/SObjectMixerEditorListRow.h
+++ b/Engine/Plugins/Editor/ObjectMixer/ObjectMixer/Source/ObjectMixer/Public/Views/List/SObjectMixerEditorListRow.h
@@ -46,6 +46,7 @@ private:
 	
 	/** Get the brush for this widget */
 	const FSlateBrush* GetVisibilityBrush() const;
+	const FSlateBrush* GetSoloBrush() const;
 
 	TSharedPtr<SWidget> GenerateCells(const FName& InColumnName, const TSharedPtr<FObjectMixerEditorListRow> PinnedItem);
 
@@ -60,6 +61,9 @@ private:
 	const FSlateBrush* VisibleNotHoveredBrush = nullptr;
 	const FSlateBrush* NotVisibleHoveredBrush = nullptr;
 	const FSlateBrush* NotVisibleNotHoveredBrush = nullptr;
+	
+	const FSlateBrush* SoloOnBrush = nullptr;
+	const FSlateBrush* SoloOffHoveredBrush = nullptr;
 
 	/** The offset applied to text widgets so that the text aligns with the column header text */
 	float TextBlockLeftPadding = 3.0f;
