Object Mixer:

-Fix: Removing anything from a collection deletes the entire collection

diff --git a/Engine/Plugins/Editor/ObjectMixer/ObjectMixer/Source/ObjectMixer/Private/ObjectMixerEditorSerializedData.cpp b/Engine/Plugins/Editor/ObjectMixer/ObjectMixer/Source/ObjectMixer/Private/ObjectMixerEditorSerializedData.cpp
--- a/Engine/Plugins/Editor/ObjectMixer/ObjectMixer/Source/ObjectMixer/Private/ObjectMixerEditorSerializedData.cpp
+++ b/Engine/Plugins/Editor/ObjectMixer/ObjectMixer/Source/ObjectMixer/Private/ObjectMixerEditorSerializedData.cpp
@@ -39,7 +39,7 @@ void UObjectMixerEditorSerializedData::AddObjectsToCollection(const FName& Filte
 	
 	if (Data)
 	{
-		TSet<FObjectMixerCollectionObjectData> NewObjectsData;
+		TArray<FObjectMixerCollectionObjectData> NewObjectsData;
 		for (const FSoftObjectPath& Object : ObjectsToAdd)
 		{
 			NewObjectsData.Add({Object});
@@ -51,6 +51,12 @@ void UObjectMixerEditorSerializedData::AddObjectsToCollection(const FName& Filte
 				return Comparator.CollectionName.IsEqual(CollectionName);
 			}))
 		{
+			NewObjectsData.SetNum(Algo::StableRemoveIf(NewObjectsData,
+				[Match](const FObjectMixerCollectionObjectData& Other)
+				{
+					return Match->CollectionObjects.Contains(Other);
+				}));
+			
 			Match->CollectionObjects.Append(NewObjectsData);
 		}
 		else
@@ -65,18 +71,27 @@ void UObjectMixerEditorSerializedData::AddObjectsToCollection(const FName& Filte
 void UObjectMixerEditorSerializedData::RemoveObjectsFromCollection(const FName& FilterClassName, const FName& CollectionName, const TSet<FSoftObjectPath>& ObjectsToRemove)
 {
 	if (FObjectMixerSerializationDataPerFilter* Data = FindSerializationDataByFilterClassName(FilterClassName))
-	{
-		FScopedTransaction RemoveCollectionTransaction(LOCTEXT("RemoveCollectionTransaction","Remove Collection"));
-	
-		Modify();
-		
-		Data->SerializedCollections.SetNum(Algo::StableRemoveIf(Data->SerializedCollections,
+	{		
+		if (FObjectMixerCollectionObjectSet* Match = Algo::FindByPredicate(Data->SerializedCollections,
 			[CollectionName](const FObjectMixerCollectionObjectSet& Comparator)
 			{
 				return Comparator.CollectionName.IsEqual(CollectionName);
+			}))
+		{
+			FScopedTransaction RemoveObjectsFromCollectionTransaction(
+				LOCTEXT("RemoveObjectsFromCollectionTransaction_Format","Remove {0}|plural(one=Object,other=Objects) From Collection"));
+			
+			Modify();
+
+			Match->CollectionObjects.SetNum(Algo::StableRemoveIf(Match->CollectionObjects,
+			[&ObjectsToRemove](const FObjectMixerCollectionObjectData& Comparator)
+			{
+				return ObjectsToRemove.Contains(Comparator.ObjectPath);
 			}));
+			
+			SaveConfig();
+		}
 		
-		SaveConfig();
 	}
 }
 
@@ -117,7 +132,7 @@ void UObjectMixerEditorSerializedData::ReorderCollection(const FName& FilterClas
 
 		if (bFoundMatch)
 		{
-			FScopedTransaction RemoveCollectionTransaction(LOCTEXT("RemoveCollectionTransaction","Remove Collection"));
+			FScopedTransaction ReorderCollectionTransaction(LOCTEXT("ReorderCollectionTransaction","Reorder Collection"));
 	
 			Modify();
 			
@@ -153,7 +168,7 @@ bool UObjectMixerEditorSerializedData::IsObjectInCollection(const FName& FilterC
 				return Comparator.CollectionName.IsEqual(CollectionName);
 			}))
 		{
-			return Algo::FindByPredicate((*Match).CollectionObjects,
+			return Algo::FindByPredicate(Match->CollectionObjects,
 				[InObject](const FObjectMixerCollectionObjectData& Comparator)
 				{
 					return Comparator.ObjectPath == InObject;
diff --git a/Engine/Plugins/Editor/ObjectMixer/ObjectMixer/Source/ObjectMixer/Public/ObjectMixerEditorSerializedData.h b/Engine/Plugins/Editor/ObjectMixer/ObjectMixer/Source/ObjectMixer/Public/ObjectMixerEditorSerializedData.h
--- a/Engine/Plugins/Editor/ObjectMixer/ObjectMixer/Source/ObjectMixer/Public/ObjectMixerEditorSerializedData.h
+++ b/Engine/Plugins/Editor/ObjectMixer/ObjectMixer/Source/ObjectMixer/Public/ObjectMixerEditorSerializedData.h
@@ -33,7 +33,7 @@ struct FObjectMixerCollectionObjectSet
 	FName CollectionName = NAME_None;
 
 	UPROPERTY()
-	TSet<FObjectMixerCollectionObjectData> CollectionObjects = {};
+	TArray<FObjectMixerCollectionObjectData> CollectionObjects = {};
 
 	bool operator==(const FObjectMixerCollectionObjectSet& Other) const
 	{
