Object Mixer:

-Move Object Mixer Filter to Miscellaneous Category

diff --git a/Engine/Plugins/Editor/ObjectMixer/ObjectMixer/Source/ObjectMixer/Private/ObjectMixerFilterFactory.cpp b/Engine/Plugins/Editor/ObjectMixer/ObjectMixer/Source/ObjectMixer/Private/ObjectMixerFilterFactory.cpp
--- a/Engine/Plugins/Editor/ObjectMixer/ObjectMixer/Source/ObjectMixer/Private/ObjectMixerFilterFactory.cpp
+++ b/Engine/Plugins/Editor/ObjectMixer/ObjectMixer/Source/ObjectMixer/Private/ObjectMixerFilterFactory.cpp
@@ -41,8 +41,7 @@ FString UObjectMixerBlueprintFilterFactory::GetDefaultNewAssetName() const
 
 uint32 UObjectMixerBlueprintFilterFactory::GetMenuCategories() const
 {
-	IAssetTools& AssetTools = FModuleManager::LoadModuleChecked<FAssetToolsModule>("AssetTools").Get();
-	return AssetTools.RegisterAdvancedAssetCategory("Object Mixer", LOCTEXT("AssetCategoryName", "Object Mixer"));
+	return EAssetTypeCategories::Misc;
 }
 
 #undef LOCTEXT_NAMESPACE
