Object Mixer:

-Fix: Viewport not updated when multi-editing some properties
-Move Project Settings entries to Editor Settings
-When changing collection filters, row expansion is now based on bExpandTreeViewItemsByDefault in UObjectMixerEditorSettings

diff --git a/Engine/Plugins/Editor/ObjectMixer/LightMixer/Source/LightMixer/Private/LightMixerModule.cpp b/Engine/Plugins/Editor/ObjectMixer/LightMixer/Source/LightMixer/Private/LightMixerModule.cpp
--- a/Engine/Plugins/Editor/ObjectMixer/LightMixer/Source/LightMixer/Private/LightMixerModule.cpp
+++ b/Engine/Plugins/Editor/ObjectMixer/LightMixer/Source/LightMixer/Private/LightMixerModule.cpp
@@ -3,7 +3,7 @@
 #include "LightMixerModule.h"
 
 #include "LightMixerObjectFilter.h"
-#include "LightMixerProjectSettings.h"
+#include "LightMixerEditorSettings.h"
 #include "LightMixerStyle.h"
 
 #include "ISettingsModule.h"
@@ -36,7 +36,7 @@ FLightMixerModule& FLightMixerModule::Get()
 void FLightMixerModule::OpenProjectSettings()
 {
 	FModuleManager::LoadModuleChecked<ISettingsModule>("Settings")
-		.ShowViewer("Project", "Editor", "Light Mixer");
+		.ShowViewer("Editor", "Plugins", "Light Mixer");
 }
 
 void FLightMixerModule::Initialize()
@@ -66,24 +66,24 @@ FName FLightMixerModule::GetTabSpawnerId()
 	return "LightMixerToolkit";
 }
 
-void FLightMixerModule::RegisterProjectSettings() const
+void FLightMixerModule::RegisterSettings() const
 {
 	if (ISettingsModule* SettingsModule = FModuleManager::GetModulePtr<ISettingsModule>("Settings"))
 	{
 		// User Project Settings
 		const TSharedPtr<ISettingsSection> ProjectSettingsSectionPtr = SettingsModule->RegisterSettings(
-			"Project", "Editor", "Light Mixer",
+			"Editor", "Plugins", "Light Mixer",
 			LOCTEXT("LightMixerSettingsDisplayName", "Light Mixer"),
 			LOCTEXT("LightMixerSettingsDescription", "Configure Light Mixer user settings"),
-			GetMutableDefault<ULightMixerProjectSettings>());
+			GetMutableDefault<ULightMixerEditorSettings>());
 	}
 }
 
-void FLightMixerModule::UnregisterProjectSettings() const
+void FLightMixerModule::UnregisterSettings() const
 {
 	if (ISettingsModule* SettingsModule = FModuleManager::GetModulePtr<ISettingsModule>("Settings"))
 	{
-		SettingsModule->UnregisterSettings("Project", "Editor", "Light Mixer");
+		SettingsModule->UnregisterSettings("Editor", "Plugins", "Light Mixer");
 	}
 }
 
diff --git a/Engine/Plugins/Editor/ObjectMixer/LightMixer/Source/LightMixer/Public/LightMixerEditorSettings.h b/Engine/Plugins/Editor/ObjectMixer/LightMixer/Source/LightMixer/Public/LightMixerEditorSettings.h
new file mode 100644
--- /dev/null
+++ b/Engine/Plugins/Editor/ObjectMixer/LightMixer/Source/LightMixer/Public/LightMixerEditorSettings.h
@@ -0,0 +1,22 @@
+// Copyright Epic Games, Inc. All Rights Reserved.
+
+#pragma once
+
+#include "LightMixerEditorSettings.generated.h"
+
+UCLASS(config = ObjectMixer)
+class LIGHTMIXER_API ULightMixerEditorSettings : public UObject
+{
+	GENERATED_BODY()
+public:
+	
+	ULightMixerEditorSettings(const FObjectInitializer& ObjectInitializer)
+	{}
+
+	/**
+	 * If true, the Object Mixer menu item will be disabled and removed.
+	 * This is useful if you only want to use Light Mixer and to avoid cluttering the menus.
+	 */
+	UPROPERTY(Config, EditAnywhere, Category="Light Mixer", meta = (ConfigRestartRequired = true))
+	bool bHideObjectMixerMenuItem = false;
+};
diff --git a/Engine/Plugins/Editor/ObjectMixer/LightMixer/Source/LightMixer/Public/LightMixerModule.h b/Engine/Plugins/Editor/ObjectMixer/LightMixer/Source/LightMixer/Public/LightMixerModule.h
--- a/Engine/Plugins/Editor/ObjectMixer/LightMixer/Source/LightMixer/Public/LightMixerModule.h
+++ b/Engine/Plugins/Editor/ObjectMixer/LightMixer/Source/LightMixer/Public/LightMixerModule.h
@@ -22,7 +22,7 @@ public:
 	virtual FName GetModuleName() override;
 	virtual void SetupMenuItemVariables() override;
 	virtual FName GetTabSpawnerId() override;
-	virtual void RegisterProjectSettings() const override;
-	virtual void UnregisterProjectSettings() const override;
+	virtual void RegisterSettings() const override;
+	virtual void UnregisterSettings() const override;
 	//~ End FObjectMixerEditorModule overrides
 };
diff --git a/Engine/Plugins/Editor/ObjectMixer/LightMixer/Source/LightMixer/Public/LightMixerProjectSettings.h b/Engine/Plugins/Editor/ObjectMixer/LightMixer/Source/LightMixer/Public/LightMixerProjectSettings.h
deleted file mode 100644
diff --git a/Engine/Plugins/Editor/ObjectMixer/ObjectMixer/Source/ObjectMixer/Private/ObjectMixerEditorModule.cpp b/Engine/Plugins/Editor/ObjectMixer/ObjectMixer/Source/ObjectMixer/Private/ObjectMixerEditorModule.cpp
--- a/Engine/Plugins/Editor/ObjectMixer/ObjectMixer/Source/ObjectMixer/Private/ObjectMixerEditorModule.cpp
+++ b/Engine/Plugins/Editor/ObjectMixer/ObjectMixer/Source/ObjectMixer/Private/ObjectMixerEditorModule.cpp
@@ -3,7 +3,7 @@
 #include "ObjectMixerEditorModule.h"
 
 #include "ObjectMixerEditorLog.h"
-#include "ObjectMixerEditorProjectSettings.h"
+#include "ObjectMixerEditorSettings.h"
 #include "ObjectMixerEditorStyle.h"
 #include "Views/MainPanel/ObjectMixerEditorMainPanel.h"
 
@@ -51,7 +51,7 @@ void FObjectMixerEditorModule::Initialize()
 	SetupMenuItemVariables();
 	
 	RegisterTabSpawner();
-	RegisterProjectSettings();
+	RegisterSettings();
 }
 
 void FObjectMixerEditorModule::Teardown()
@@ -97,7 +97,7 @@ void FObjectMixerEditorModule::Teardown()
 	UToolMenus::UnregisterOwner(this);
 	
 	UnregisterTabSpawner();
-	UnregisterProjectSettings();
+	UnregisterSettings();
 }
 
 FObjectMixerEditorModule& FObjectMixerEditorModule::Get()
@@ -108,7 +108,7 @@ FObjectMixerEditorModule& FObjectMixerEditorModule::Get()
 void FObjectMixerEditorModule::OpenProjectSettings()
 {
 	FModuleManager::LoadModuleChecked<ISettingsModule>("Settings")
-		.ShowViewer("Project", "Editor", "Object Mixer");
+		.ShowViewer("Editor", "Plugins", "Object Mixer");
 }
 
 FName FObjectMixerEditorModule::GetModuleName()
@@ -226,24 +226,24 @@ void FObjectMixerEditorModule::UnregisterTabSpawner()
 	FGlobalTabmanager::Get()->UnregisterNomadTabSpawner(GetTabSpawnerId());
 }
 
-void FObjectMixerEditorModule::RegisterProjectSettings() const
+void FObjectMixerEditorModule::RegisterSettings() const
 {
 	if (ISettingsModule* SettingsModule = FModuleManager::GetModulePtr<ISettingsModule>("Settings"))
 	{
 		// User Project Settings
 		const TSharedPtr<ISettingsSection> ProjectSettingsSectionPtr = SettingsModule->RegisterSettings(
-			"Project", "Editor", "Object Mixer",
+			"Editor", "Plugins", "Object Mixer",
 			LOCTEXT("ObjectMixerSettingsDisplayName", "Object Mixer"),
 			LOCTEXT("ObjectMixerSettingsDescription", "Configure Object Mixer user settings"),
-			GetMutableDefault<UObjectMixerEditorProjectSettings>());	
+			GetMutableDefault<UObjectMixerEditorSettings>());	
 	}
 }
 
-void FObjectMixerEditorModule::UnregisterProjectSettings() const
+void FObjectMixerEditorModule::UnregisterSettings() const
 {
 	if (ISettingsModule* SettingsModule = FModuleManager::GetModulePtr<ISettingsModule>("Settings"))
 	{
-		SettingsModule->UnregisterSettings("Project", "Plugins", "Object Mixer");
+		SettingsModule->UnregisterSettings("Editor", "Plugins", "Object Mixer");
 	}
 }
 
diff --git a/Engine/Plugins/Editor/ObjectMixer/ObjectMixer/Source/ObjectMixer/Private/Views/List/ObjectMixerEditorListRow.cpp b/Engine/Plugins/Editor/ObjectMixer/ObjectMixer/Source/ObjectMixer/Private/Views/List/ObjectMixerEditorListRow.cpp
--- a/Engine/Plugins/Editor/ObjectMixer/ObjectMixer/Source/ObjectMixer/Private/Views/List/ObjectMixerEditorListRow.cpp
+++ b/Engine/Plugins/Editor/ObjectMixer/ObjectMixer/Source/ObjectMixer/Private/Views/List/ObjectMixerEditorListRow.cpp
@@ -97,18 +97,23 @@ void FObjectMixerEditorListRow::SetRowType(EObjectMixerEditorListRowType InNewRo
 	RowType = InNewRowType;
 }
 
-int32 FObjectMixerEditorListRow::GetHybridRowIndex() const
+int32 FObjectMixerEditorListRow::GetOrFindHybridRowIndex()
 {
+	if (CachedHybridRowIndex != INDEX_NONE)
+	{
+		return CachedHybridRowIndex;
+	}
+		
 	if (GetRowType() == FObjectMixerEditorListRow::ContainerObject)
 	{
 		if (const UObject* ThisObject = GetObject())
 		{
 			TArray<int32> HybridCandidates;
-		
+	
 			for (int32 ChildItr = 0; ChildItr < GetChildCount(); ChildItr++)
 			{
 				const FObjectMixerEditorListRowPtr& ChildRow = GetChildRows()[ChildItr];
-			
+		
 				if (ChildRow->GetRowType() == FObjectMixerEditorListRow::MatchingObject)
 				{
 					if (const UObject* ChildObject = ChildRow->GetObject(); ChildObject->GetOuter() == ThisObject)
@@ -124,10 +129,11 @@ int32 FObjectMixerEditorListRow::GetHybridRowIndex() const
 					}
 				}
 			}
-			
+		
 			if (HybridCandidates.Num() == 1)
 			{
-				return HybridCandidates[0];
+				CachedHybridRowIndex = HybridCandidates[0];
+				return CachedHybridRowIndex;
 			}
 		}
 	}
@@ -135,6 +141,12 @@ int32 FObjectMixerEditorListRow::GetHybridRowIndex() const
 	return INDEX_NONE;
 }
 
+FObjectMixerEditorListRowPtr FObjectMixerEditorListRow::GetHybridChild()
+{
+	const int32 HybridIndex = GetOrFindHybridRowIndex();
+	return HybridIndex != INDEX_NONE ? GetChildRows()[HybridIndex] : nullptr;
+}
+
 int32 FObjectMixerEditorListRow::GetSortOrder() const
 {
 	return SortOrder;
diff --git a/Engine/Plugins/Editor/ObjectMixer/ObjectMixer/Source/ObjectMixer/Private/Views/List/SObjectMixerEditorList.cpp b/Engine/Plugins/Editor/ObjectMixer/ObjectMixer/Source/ObjectMixer/Private/Views/List/SObjectMixerEditorList.cpp
--- a/Engine/Plugins/Editor/ObjectMixer/ObjectMixer/Source/ObjectMixer/Private/Views/List/SObjectMixerEditorList.cpp
+++ b/Engine/Plugins/Editor/ObjectMixer/ObjectMixer/Source/ObjectMixer/Private/Views/List/SObjectMixerEditorList.cpp
@@ -7,7 +7,7 @@
 
 #include "ObjectFilter/ObjectMixerEditorObjectFilter.h"
 #include "ObjectMixerEditorLog.h"
-#include "ObjectMixerEditorProjectSettings.h"
+#include "ObjectMixerEditorSettings.h"
 #include "ObjectMixerEditorStyle.h"
 #include "Views/List/SObjectMixerEditorListRow.h"
 
@@ -90,7 +90,7 @@ void SObjectMixerEditorList::Construct(const FArguments& InArgs, TSharedRef<FObj
 				{
 					if (SelectionType != ESelectInfo::Direct) // Don't call if selected from code
 					{
-						const bool bSyncSelectionEnabled = GetDefault<UObjectMixerEditorProjectSettings>()->bSyncSelection;
+						const bool bSyncSelectionEnabled = GetDefault<UObjectMixerEditorSettings>()->bSyncSelection;
 						const bool bIsAltDown = FSlateApplication::Get().GetModifierKeys().IsAltDown();
 						const bool bShouldSyncSelection = bSyncSelectionEnabled ? !bIsAltDown : bIsAltDown;
 						
@@ -211,7 +211,7 @@ int32 SObjectMixerEditorList::GetSelectedTreeViewItemCount() const
 
 void SObjectMixerEditorList::SyncEditorSelectionToListSelection()
 {
-	if (GetDefault<UObjectMixerEditorProjectSettings>()->bSyncSelection && !bShouldPauseSyncSelection)
+	if (GetDefault<UObjectMixerEditorSettings>()->bSyncSelection && !bShouldPauseSyncSelection)
 	{
 		bIsEditorToListSelectionSyncRequested = false;
 		const USelection* SelectedActors = GEditor->GetSelectedActors();
@@ -415,6 +415,12 @@ void SObjectMixerEditorList::EvaluateIfRowsPassFilters(const bool bShouldRefresh
 	const TFunction<void(const TArray<FObjectMixerEditorListRowPtr>&)> RecursiveRowIterator =
 		[&MatchAnyOfFilters, &MatchAllOfFilters, &RecursiveRowIterator](const TArray<FObjectMixerEditorListRowPtr>& InObjects)
 	{
+		bool bExpandByDefault = true;
+		if (const UObjectMixerEditorSettings* Settings = GetDefault<UObjectMixerEditorSettings>())
+		{
+			bExpandByDefault = Settings->bExpandTreeViewItemsByDefault;
+		}
+			
 		for (const FObjectMixerEditorListRowPtr& Row : InObjects)
 		{
 			if (Row.IsValid())
@@ -433,6 +439,8 @@ void SObjectMixerEditorList::EvaluateIfRowsPassFilters(const bool bShouldRefresh
 				{
 					RecursiveRowIterator(Row->GetChildRows());
 				}
+				
+				Row->SetIsTreeViewItemExpanded(bExpandByDefault);
 			}
 		}
 	};
@@ -952,7 +960,7 @@ void SObjectMixerEditorList::RestoreTreeState(const bool bFlushCache)
 	};
 	
 	bool bExpandByDefault = true;
-	if (const UObjectMixerEditorProjectSettings* const Settings = GetDefault<UObjectMixerEditorProjectSettings>())
+	if (const UObjectMixerEditorSettings* Settings = GetDefault<UObjectMixerEditorSettings>())
 	{
 		bExpandByDefault = Settings->bExpandTreeViewItemsByDefault;
 	}
@@ -1309,7 +1317,7 @@ void SObjectMixerEditorList::OnGetRowChildren(FObjectMixerEditorListRowPtr Row,
 	{
 		OutChildren = Row->GetChildRows();
 		
-		if (const int32 HybridIndex = Row->GetHybridRowIndex(); HybridIndex != INDEX_NONE)
+		if (const int32 HybridIndex = Row->GetOrFindHybridRowIndex(); HybridIndex != INDEX_NONE)
 		{
 			OutChildren.RemoveAt(HybridIndex);
 		}
diff --git a/Engine/Plugins/Editor/ObjectMixer/ObjectMixer/Source/ObjectMixer/Private/Views/List/SObjectMixerEditorListRow.cpp b/Engine/Plugins/Editor/ObjectMixer/ObjectMixer/Source/ObjectMixer/Private/Views/List/SObjectMixerEditorListRow.cpp
--- a/Engine/Plugins/Editor/ObjectMixer/ObjectMixer/Source/ObjectMixer/Private/Views/List/SObjectMixerEditorListRow.cpp
+++ b/Engine/Plugins/Editor/ObjectMixer/ObjectMixer/Source/ObjectMixer/Private/Views/List/SObjectMixerEditorListRow.cpp
@@ -36,7 +36,7 @@ void SObjectMixerEditorListRow::Construct(
 	check(InRow.IsValid());
 
 	Item = InRow;
-	HybridRowIndex = Item.Pin()->GetHybridRowIndex();
+	HybridRowIndex = Item.Pin()->GetOrFindHybridRowIndex();
 
 	SMultiColumnTableRow<FObjectMixerEditorListRowPtr>::Construct(
 		FSuperRowType::FArguments()
@@ -65,11 +65,9 @@ void SObjectMixerEditorListRow::Construct(
 TSharedRef<SWidget> SObjectMixerEditorListRow::GenerateWidgetForColumn(const FName& InColumnName)
 {
 	check(Item.IsValid());
+	const FObjectMixerEditorListRowPtr RowPtr = GetHybridChildOrRowItemIfNull();
 
-	const FObjectMixerEditorListRowPtr PinnedItem =
-		HybridRowIndex != INDEX_NONE ? Item.Pin()->GetChildRows()[HybridRowIndex] : Item.Pin();
-
-	if (const TSharedPtr<SWidget> CellWidget = GenerateCells(InColumnName, PinnedItem))
+	if (const TSharedPtr<SWidget> CellWidget = GenerateCells(InColumnName, RowPtr))
 	{
 		if (InColumnName == SObjectMixerEditorList::ItemNameColumnName)
 		{			
@@ -245,6 +243,34 @@ FReply SObjectMixerEditorListRow::HandleAcceptDrop(const FDragDropEvent& DragDro
 	return FReply::Handled();
 }
 
+FObjectMixerEditorListRowPtr SObjectMixerEditorListRow::GetHybridChildOrRowItemIfNull() const
+{
+	if (TSharedPtr<FObjectMixerEditorListRow> PinnedItem = Item.Pin())
+	{
+		if (FObjectMixerEditorListRowPtr HybridChild = HybridRowIndex != INDEX_NONE ? PinnedItem->GetChildRows()[HybridRowIndex] : nullptr)
+		{
+			return HybridChild;
+		}
+			
+		return PinnedItem;
+	}
+
+	return nullptr;
+}
+
+bool SObjectMixerEditorListRow::GetIsItemOrHybridChildSelected() const
+{
+	if (Item.IsValid())
+	{
+		const bool bIsItemSelected = Item.Pin()->GetIsSelected();
+		const bool bHasHybridChild = HybridRowIndex != INDEX_NONE && Item.Pin()->GetChildRows()[HybridRowIndex].IsValid();
+		const bool bIsChildSelected = bHasHybridChild && Item.Pin()->GetChildRows()[HybridRowIndex]->GetIsSelected();
+		return bIsChildSelected  || bIsItemSelected;
+	}
+
+	return false;
+}
+
 bool SObjectMixerEditorListRow::IsVisible() const
 {
 	if (const TSharedPtr<FObjectMixerEditorListRow> PinnedItem = Item.Pin())
@@ -317,11 +343,11 @@ const FSlateBrush* SObjectMixerEditorListRow::GetSoloBrush() const
 }
 
 TSharedPtr<SWidget> SObjectMixerEditorListRow::GenerateCells(
-	const FName& InColumnName, const TSharedPtr<FObjectMixerEditorListRow> PinnedItem)
+	const FName& InColumnName, const TSharedPtr<FObjectMixerEditorListRow> RowPtr)
 {
-	check(PinnedItem.IsValid());
+	check(RowPtr.IsValid());
 	
-	if (PinnedItem->GetRowType() == FObjectMixerEditorListRow::None)
+	if (RowPtr->GetRowType() == FObjectMixerEditorListRow::None)
 	{
 		return SNullWidget::NullWidget;
 	}
@@ -336,17 +362,17 @@ TSharedPtr<SWidget> SObjectMixerEditorListRow::GenerateCells(
 		.AutoWidth()
 		[
 			SNew(SImage)
-			.Image_Lambda([PinnedItem]()
+			.Image_Lambda([RowPtr]()
 			{
-				return PinnedItem->GetObjectIconBrush();
+				return RowPtr->GetObjectIconBrush();
 			})
 			.ColorAndOpacity(FSlateColor::UseForeground())
 		];
 
 		bool bNeedsStandardTextBlock = true;
-		const FText DisplayName = PinnedItem->GetDisplayName(bIsHybridRow);
+		const FText DisplayName = RowPtr->GetDisplayName(bIsHybridRow);
 		
-		if (TObjectPtr<UObject> Object = PinnedItem->GetObject())
+		if (TObjectPtr<UObject> Object = RowPtr->GetObject())
 		{
 			if (UClass* ActorClass = Object->GetClass())
 			{
@@ -373,9 +399,9 @@ TSharedPtr<SWidget> SObjectMixerEditorListRow::GenerateCells(
 		{
 			FText TooltipText = DisplayName;
 
-			if (const UObjectMixerObjectFilter* Filter = PinnedItem->GetObjectFilter())
+			if (const UObjectMixerObjectFilter* Filter = RowPtr->GetObjectFilter())
 			{
-				if (const TObjectPtr<UObject> Object = PinnedItem->GetObject())
+				if (const TObjectPtr<UObject> Object = RowPtr->GetObject())
 				{
 					TooltipText = Filter->GetRowTooltipText(Object, bIsHybridRow);
 				}
@@ -405,7 +431,7 @@ TSharedPtr<SWidget> SObjectMixerEditorListRow::GenerateCells(
 
 	if (InColumnName.IsEqual(SObjectMixerEditorList::EditorVisibilityColumnName))
 	{
-		if (!bIsHybridRow && PinnedItem->GetObject() && !PinnedItem->GetObject()->IsA(AActor::StaticClass()))
+		if (!bIsHybridRow && RowPtr->GetObject() && !RowPtr->GetObject()->IsA(AActor::StaticClass()))
 		{
 			return nullptr;
 		}
@@ -420,16 +446,16 @@ TSharedPtr<SWidget> SObjectMixerEditorListRow::GenerateCells(
 					.ColorAndOpacity(this, &SObjectMixerEditorListRow::GetVisibilityIconForegroundColor)
 					.Image_Raw(this, &SObjectMixerEditorListRow::GetVisibilityBrush)
 					.OnMouseButtonDown_Lambda(
-						[PinnedItem] (const FGeometry& MyGeometry, const FPointerEvent& Event)
+						[RowPtr] (const FGeometry& MyGeometry, const FPointerEvent& Event)
 						{
-							check (PinnedItem);
+							check (RowPtr);
 
 							FScopedTransaction Transaction( LOCTEXT("VisibilityChanged", "Object Mixer - Visibility Changed") );
 
-							const bool bIsVisible = PinnedItem->GetObjectVisibility();
+							const bool bIsVisible = RowPtr->GetObjectVisibility();
 							
-							if (const TSharedPtr<SObjectMixerEditorList> PinnedListView = PinnedItem->GetListViewPtr().Pin();
-								PinnedItem->GetIsSelected() && PinnedListView->GetSelectedTreeViewItemCount() > 0)
+							if (const TSharedPtr<SObjectMixerEditorList> PinnedListView = RowPtr->GetListViewPtr().Pin();
+								RowPtr->GetIsSelected() && PinnedListView->GetSelectedTreeViewItemCount() > 0)
 							{
 								PinnedListView->SetSelectedTreeViewItemActorsEditorVisible(!bIsVisible);
 
@@ -437,7 +463,7 @@ TSharedPtr<SWidget> SObjectMixerEditorListRow::GenerateCells(
 							}
 
 							// Set Visibility Recursively
-							PinnedItem->SetObjectVisibility(!bIsVisible, true);
+							RowPtr->SetObjectVisibility(!bIsVisible, true);
 
 							return FReply::Handled();
 						}
@@ -448,7 +474,7 @@ TSharedPtr<SWidget> SObjectMixerEditorListRow::GenerateCells(
 
 	if (InColumnName.IsEqual(SObjectMixerEditorList::EditorVisibilitySoloColumnName))
 	{
-		if (!bIsHybridRow && PinnedItem->GetObject() && !PinnedItem->GetObject()->IsA(AActor::StaticClass()))
+		if (!bIsHybridRow && RowPtr->GetObject() && !RowPtr->GetObject()->IsA(AActor::StaticClass()))
 		{
 			return nullptr;
 		}
@@ -496,7 +522,7 @@ TSharedPtr<SWidget> SObjectMixerEditorListRow::GenerateCells(
 			;
 	}
 	
-	if (UObject* ObjectRef = PinnedItem->GetObject())
+	if (UObject* ObjectRef = RowPtr->GetObject())
 	{
 		FPropertyEditorModule& PropertyEditorModule = FModuleManager::LoadModuleChecked<FPropertyEditorModule>("PropertyEditor");
 		{
@@ -513,20 +539,17 @@ TSharedPtr<SWidget> SObjectMixerEditorListRow::GenerateCells(
 				{
 					if (const FProperty* Property = Handle->GetProperty())
 					{
+						const FName PropertyName = Property->GetFName();
+						RowPtr->PropertyNamesToHandles.Add(PropertyName, Handle);
+						
 						// Simultaneously edit all selected rows with a similar property
 						FSimpleDelegate OnPropertyValueChanged =
 							FSimpleDelegate::CreateRaw(
 								this,
-								&SObjectMixerEditorListRow::OnPropertyChanged, Property, (void*)ObjectRef);
+								&SObjectMixerEditorListRow::OnPropertyChanged, PropertyName);
 					
 						Handle->SetOnPropertyValueChanged(OnPropertyValueChanged);
-						uint32 ChildHandleCount;
-						Handle->GetNumChildren(ChildHandleCount);
-
-						for (uint32 ChildHandleItr = 0; ChildHandleItr < ChildHandleCount; ChildHandleItr++)
-						{
-							Handle->GetChildHandle(ChildHandleItr)->SetOnPropertyValueChanged(OnPropertyValueChanged);
-						}
+						Handle->SetOnChildPropertyValueChanged(OnPropertyValueChanged);
 
 						return SNew(SBox)
 								.Visibility(EVisibility::SelfHitTestInvisible)
@@ -544,31 +567,61 @@ TSharedPtr<SWidget> SObjectMixerEditorListRow::GenerateCells(
 	return nullptr;
 }
 
-void SObjectMixerEditorListRow::OnPropertyChanged(const FProperty* Property, void* ContainerWithChangedProperty) const
+void SObjectMixerEditorListRow::OnPropertyChanged(const FName PropertyName) const
 {
-	if (Property && ContainerWithChangedProperty)
+	check(Item.IsValid());
+
+	struct Local
 	{
-		if (const TSharedPtr<FObjectMixerEditorListRow> PinnedItem = Item.Pin())
+		static void SetValueOnSelectedItems(
+			const FString& ValueAsString, const TArray<FObjectMixerEditorListRowPtr>& OtherSelectedItems,
+			const FName& PropertyName, const FObjectMixerEditorListRowPtr PinnedItem)
 		{
-			if (PinnedItem->GetIsSelected())
+			if (!ValueAsString.IsEmpty())
 			{
-				if (const void* ValuePtr = Property->ContainerPtrToValuePtr<void>(ContainerWithChangedProperty))
+				FScopedTransaction Transaction(
+					LOCTEXT("OnPropertyChangedTransaction", "Object Mixer - Bulk Edit Selected Row Properties") );
+						
+				for (const TSharedPtr<FObjectMixerEditorListRow>& SelectedRow : OtherSelectedItems)
 				{
-					FScopedTransaction Transaction( LOCTEXT("PropertyChanged", "Object Mixer - Property Changed") );
-					
-					for (const TSharedPtr<FObjectMixerEditorListRow>& SelectedRow : PinnedItem->GetSelectedTreeViewItems())
-					{
-						UObject* SelectedRowObject = SelectedRow->GetObject();
-						const int32 SelectedRowHybridIndex = SelectedRow->GetHybridRowIndex();
+					const FObjectMixerEditorListRowPtr SelectedHybridRow = SelectedRow->GetHybridChild();
+					const FObjectMixerEditorListRowPtr RowToUse = SelectedHybridRow.IsValid() ? SelectedHybridRow : SelectedRow;
 
-						if (SelectedRowHybridIndex != -1 && SelectedRowHybridIndex < SelectedRow->GetChildCount())
+					if (RowToUse != PinnedItem)
+					{
+						if (const TWeakPtr<IPropertyHandle>* SelectedHandlePtr = RowToUse->PropertyNamesToHandles.Find(PropertyName))
 						{
-							SelectedRowObject = SelectedRow->GetChildRows()[SelectedRowHybridIndex]->GetObject();
+							if (SelectedHandlePtr->IsValid())
+							{
+								if (UObject* ObjectToModify = RowToUse->GetObject())
+								{
+									ObjectToModify->Modify();
+								}
+										
+								SelectedHandlePtr->Pin()->SetValueFromFormattedString(ValueAsString);
+							}
 						}
-						
-						if (SelectedRowObject)
+					}
+				}
+			}
+		}
+	};
+	
+	if (PropertyName != NAME_None)
+	{
+		if (const FObjectMixerEditorListRowPtr PinnedItem = GetHybridChildOrRowItemIfNull())
+		{
+			if (GetIsItemOrHybridChildSelected())
+			{
+				if (const TWeakPtr<IPropertyHandle>* HandlePtr = PinnedItem->PropertyNamesToHandles.Find(PropertyName); HandlePtr->IsValid())
+				{
+					if (TArray<FObjectMixerEditorListRowPtr> OtherSelectedItems = PinnedItem->GetSelectedTreeViewItems(); OtherSelectedItems.Num())
+					{
+						FString ValueAsString;
+						if (TSharedPtr<IPropertyHandle> PinnedHandle = HandlePtr->Pin())
 						{
-							Property->SetValue_InContainer(SelectedRowObject, ValuePtr);
+							PinnedHandle->GetValueAsFormattedString(ValueAsString);
+							Local::SetValueOnSelectedItems(ValueAsString, OtherSelectedItems, PropertyName, PinnedItem);
 						}
 					}
 				}
diff --git a/Engine/Plugins/Editor/ObjectMixer/ObjectMixer/Source/ObjectMixer/Private/Views/MainPanel/SObjectMixerEditorMainPanel.cpp b/Engine/Plugins/Editor/ObjectMixer/ObjectMixer/Source/ObjectMixer/Private/Views/MainPanel/SObjectMixerEditorMainPanel.cpp
--- a/Engine/Plugins/Editor/ObjectMixer/ObjectMixer/Source/ObjectMixer/Private/Views/MainPanel/SObjectMixerEditorMainPanel.cpp
+++ b/Engine/Plugins/Editor/ObjectMixer/ObjectMixer/Source/ObjectMixer/Private/Views/MainPanel/SObjectMixerEditorMainPanel.cpp
@@ -4,7 +4,7 @@
 
 #include "ObjectMixerEditorLog.h"
 #include "ObjectMixerEditorModule.h"
-#include "ObjectMixerEditorProjectSettings.h"
+#include "ObjectMixerEditorSettings.h"
 #include "Views/List/ObjectMixerEditorList.h"
 #include "Views/MainPanel/ObjectMixerEditorMainPanel.h"
 #include "Views/Widgets/SCollectionSelectionButton.h"
@@ -109,7 +109,7 @@ TSharedRef<SWidget> SObjectMixerEditorMainPanel::GenerateToolbar()
 		.Style(&FAppStyle::Get().GetWidgetStyle<FCheckBoxStyle>("ToggleButtonCheckbox"))
 		.IsChecked_Lambda([]()
 		{
-			if (const UObjectMixerEditorProjectSettings* Settings = GetDefault<UObjectMixerEditorProjectSettings>())
+			if (const UObjectMixerEditorSettings* Settings = GetDefault<UObjectMixerEditorSettings>())
 			{
 				return Settings->bSyncSelection ? ECheckBoxState::Checked : ECheckBoxState::Unchecked;
 			}
@@ -118,7 +118,7 @@ TSharedRef<SWidget> SObjectMixerEditorMainPanel::GenerateToolbar()
 		})
 		.OnCheckStateChanged_Lambda([](ECheckBoxState InNewState)
 		{
-			if (UObjectMixerEditorProjectSettings* Settings = GetMutableDefault<UObjectMixerEditorProjectSettings>())
+			if (UObjectMixerEditorSettings* Settings = GetMutableDefault<UObjectMixerEditorSettings>())
 			{
 				Settings->bSyncSelection = InNewState == ECheckBoxState::Checked ? true : false;
 			}
diff --git a/Engine/Plugins/Editor/ObjectMixer/ObjectMixer/Source/ObjectMixer/Public/ObjectMixerEditorModule.h b/Engine/Plugins/Editor/ObjectMixer/ObjectMixer/Source/ObjectMixer/Public/ObjectMixerEditorModule.h
--- a/Engine/Plugins/Editor/ObjectMixer/ObjectMixer/Source/ObjectMixer/Public/ObjectMixerEditorModule.h
+++ b/Engine/Plugins/Editor/ObjectMixer/ObjectMixer/Source/ObjectMixer/Public/ObjectMixerEditorModule.h
@@ -54,8 +54,8 @@ public:
 	bool RegisterItemInMenuGroup(FWorkspaceItem& InItem);
 	
 	virtual void UnregisterTabSpawner();
-	virtual void RegisterProjectSettings() const;
-	virtual void UnregisterProjectSettings() const;
+	virtual void RegisterSettings() const;
+	virtual void UnregisterSettings() const;
 
 	virtual TSharedRef<SDockTab> SpawnTab(const FSpawnTabArgs& Args);
 
diff --git a/Engine/Plugins/Editor/ObjectMixer/ObjectMixer/Source/ObjectMixer/Public/ObjectMixerEditorProjectSettings.h b/Engine/Plugins/Editor/ObjectMixer/ObjectMixer/Source/ObjectMixer/Public/ObjectMixerEditorProjectSettings.h
deleted file mode 100644
diff --git a/Engine/Plugins/Editor/ObjectMixer/ObjectMixer/Source/ObjectMixer/Public/ObjectMixerEditorSettings.h b/Engine/Plugins/Editor/ObjectMixer/ObjectMixer/Source/ObjectMixer/Public/ObjectMixerEditorSettings.h
new file mode 100644
--- /dev/null
+++ b/Engine/Plugins/Editor/ObjectMixer/ObjectMixer/Source/ObjectMixer/Public/ObjectMixerEditorSettings.h
@@ -0,0 +1,29 @@
+// Copyright Epic Games, Inc. All Rights Reserved.
+
+#pragma once
+
+#include "ObjectMixerEditorSettings.generated.h"
+
+UCLASS(config = ObjectMixer)
+class OBJECTMIXEREDITOR_API UObjectMixerEditorSettings : public UObject
+{
+	GENERATED_BODY()
+public:
+	
+	UObjectMixerEditorSettings(const FObjectInitializer& ObjectInitializer)
+	{}
+	
+	/**
+	 * If enabled, clicking an item in the mixer list will also select the item in the Scene Outliner.
+	 * Alt + Click to select items in mixer without selecting the item in the Scene outliner.
+	 * If disabled, selections will not sync unless Alt is held. Effectively, this is the opposite behavior.
+	 */
+	UPROPERTY(Config, EditAnywhere, Category="Object Mixer")
+	bool bSyncSelection = true;
+
+	/**
+	 * If false, a new object will be created every time the filter object is accessed.
+	 */
+	UPROPERTY(Config, EditAnywhere, Category="Object Mixer")
+	bool bExpandTreeViewItemsByDefault = true;
+};
diff --git a/Engine/Plugins/Editor/ObjectMixer/ObjectMixer/Source/ObjectMixer/Public/Views/List/ObjectMixerEditorListRow.h b/Engine/Plugins/Editor/ObjectMixer/ObjectMixer/Source/ObjectMixer/Public/Views/List/ObjectMixerEditorListRow.h
--- a/Engine/Plugins/Editor/ObjectMixer/ObjectMixer/Source/ObjectMixer/Public/Views/List/ObjectMixerEditorListRow.h
+++ b/Engine/Plugins/Editor/ObjectMixer/ObjectMixer/Source/ObjectMixer/Public/Views/List/ObjectMixerEditorListRow.h
@@ -7,6 +7,7 @@
 #include "DragAndDrop/DecoratedDragDropOp.h"
 #include "Layout/Visibility.h"
 #include "ObjectFilter/ObjectMixerEditorObjectFilter.h"
+#include "PropertyHandle.h"
 #include "Templates/SharedPointer.h"
 #include "UObject/Object.h"
 #include "UObject/WeakObjectPtr.h"
@@ -102,8 +103,10 @@ struct OBJECTMIXEREDITOR_API FObjectMixerEditorListRow final : TSharedFromThis<F
 	[[nodiscard]] EObjectMixerEditorListRowType GetRowType() const;
 	void SetRowType(EObjectMixerEditorListRowType InNewRowType);
 
-	/** If this is a hybrid row, return the index of the child used to hybrid with. */
-	[[nodiscard]] int32 GetHybridRowIndex() const;
+	/** If this is a hybrid row, return the index of the child row used to hybrid with. */
+	[[nodiscard]] int32 GetOrFindHybridRowIndex();
+	/** If this is a hybrid row, return the child row used to hybrid with. */
+	[[nodiscard]] FObjectMixerEditorListRowPtr GetHybridChild();
 
 	[[nodiscard]] int32 GetSortOrder() const;
 	void SetSortOrder(const int32 InNewOrder);
@@ -179,6 +182,8 @@ struct OBJECTMIXEREDITOR_API FObjectMixerEditorListRow final : TSharedFromThis<F
 	void SetRowSoloState(const bool bNewSolo);
 
 	void ClearSoloRows() const;
+	
+	TMap<FName, TWeakPtr<IPropertyHandle>> PropertyNamesToHandles;
 
 private:
 
@@ -205,4 +210,6 @@ private:
 
 	// Used to expand all children on shift+click.
 	bool bShouldExpandAllChildren = false;
+
+	int32 CachedHybridRowIndex = INDEX_NONE;
 };
diff --git a/Engine/Plugins/Editor/ObjectMixer/ObjectMixer/Source/ObjectMixer/Public/Views/List/SObjectMixerEditorListRow.h b/Engine/Plugins/Editor/ObjectMixer/ObjectMixer/Source/ObjectMixer/Public/Views/List/SObjectMixerEditorListRow.h
--- a/Engine/Plugins/Editor/ObjectMixer/ObjectMixer/Source/ObjectMixer/Public/Views/List/SObjectMixerEditorListRow.h
+++ b/Engine/Plugins/Editor/ObjectMixer/ObjectMixer/Source/ObjectMixer/Public/Views/List/SObjectMixerEditorListRow.h
@@ -36,6 +36,10 @@ public:
 	void HandleDragLeave(const FDragDropEvent& DragDropEvent);
 	TOptional<EItemDropZone> HandleCanAcceptDrop(const FDragDropEvent& DragDropEvent, EItemDropZone DropZone, FObjectMixerEditorListRowPtr TargetItem);
 	FReply HandleAcceptDrop(const FDragDropEvent& DragDropEvent, EItemDropZone DropZone, FObjectMixerEditorListRowPtr TargetItem);
+
+	FObjectMixerEditorListRowPtr GetHybridChildOrRowItemIfNull() const;
+
+	bool GetIsItemOrHybridChildSelected() const;
 	
 private:
 
@@ -48,9 +52,9 @@ private:
 	const FSlateBrush* GetVisibilityBrush() const;
 	const FSlateBrush* GetSoloBrush() const;
 
-	TSharedPtr<SWidget> GenerateCells(const FName& InColumnName, const TSharedPtr<FObjectMixerEditorListRow> PinnedItem);
+	TSharedPtr<SWidget> GenerateCells(const FName& InColumnName, const TSharedPtr<FObjectMixerEditorListRow> RowPtr);
 
-	void OnPropertyChanged(const FProperty* Property, void* ContainerWithChangedProperty) const;
+	void OnPropertyChanged(const FName PropertyName) const;
 
 	void OnClickBlueprintLink(UBlueprint* AsBlueprint, UObject* Object = nullptr);
 	FText GetHighlightText() const;
