Object Mixer:

-Light Collection RMB Menu

diff --git a/Engine/Plugins/Editor/ObjectMixer/ObjectMixer/Source/ObjectMixer/Private/ObjectMixerEditorSerializedData.cpp b/Engine/Plugins/Editor/ObjectMixer/ObjectMixer/Source/ObjectMixer/Private/ObjectMixerEditorSerializedData.cpp
--- a/Engine/Plugins/Editor/ObjectMixer/ObjectMixer/Source/ObjectMixer/Private/ObjectMixerEditorSerializedData.cpp
+++ b/Engine/Plugins/Editor/ObjectMixer/ObjectMixer/Source/ObjectMixer/Private/ObjectMixerEditorSerializedData.cpp
@@ -3,12 +3,14 @@
 #include "ObjectMixerEditorSerializedData.h"
 
 #include "Algo/Find.h"
+#include "Algo/IndexOf.h"
 #include "Algo/RemoveIf.h"
 #include "ScopedTransaction.h"
-#include "Algo/IndexOf.h"
 
 #define LOCTEXT_NAMESPACE "ObjectMixerEditor"
 
+FName UObjectMixerEditorSerializedData::AllCollectionName = TEXT("All");
+
 FObjectMixerSerializationDataPerFilter* UObjectMixerEditorSerializedData::FindSerializationDataByFilterClassName(const FName& FilterClassName)
 {
 	if (FObjectMixerSerializationDataPerFilter* Match = Algo::FindByPredicate(SerializedDataPerFilter,
@@ -23,12 +25,19 @@ FObjectMixerSerializationDataPerFilter* UObjectMixerEditorSerializedData::FindSe
 	return nullptr;
 }
 
-void UObjectMixerEditorSerializedData::AddObjectsToCollection(const FName& FilterClassName, const FName& CollectionName, const TSet<FSoftObjectPath>& ObjectsToAdd)
+bool UObjectMixerEditorSerializedData::AddObjectsToCollection(const FName& FilterClassName, const FName& CollectionName, const TSet<FSoftObjectPath>& ObjectsToAdd)
 {
+	if (CollectionName == NAME_None ||
+		CollectionName.IsEqual(UObjectMixerEditorSerializedData::AllCollectionName, ENameCase::IgnoreCase) ||
+		CollectionName.ToString().TrimStartAndEnd().IsEmpty())
+	{
+		return false;
+	}
+	
 	FScopedTransaction AddObjectsToCollectionTransaction(LOCTEXT("AddObjectsToCollectionTransaction","Add Objects To Collection"));
 
 	Modify();
-	
+
 	FObjectMixerSerializationDataPerFilter* Data = FindSerializationDataByFilterClassName(FilterClassName);
 
 	if (!Data)
@@ -36,7 +45,7 @@ void UObjectMixerEditorSerializedData::AddObjectsToCollection(const FName& Filte
 		SerializedDataPerFilter.Add({FilterClassName});
 		Data = FindSerializationDataByFilterClassName(FilterClassName);
 	}
-	
+
 	if (Data)
 	{
 		TArray<FObjectMixerCollectionObjectData> NewObjectsData;
@@ -44,7 +53,7 @@ void UObjectMixerEditorSerializedData::AddObjectsToCollection(const FName& Filte
 		{
 			NewObjectsData.Add({Object});
 		}
-		
+	
 		if (FObjectMixerCollectionObjectSet* Match = Algo::FindByPredicate(Data->SerializedCollections,
 			[CollectionName](const FObjectMixerCollectionObjectSet& Comparator)
 			{
@@ -56,7 +65,7 @@ void UObjectMixerEditorSerializedData::AddObjectsToCollection(const FName& Filte
 				{
 					return Match->CollectionObjects.Contains(Other);
 				}));
-			
+		
 			Match->CollectionObjects.Append(NewObjectsData);
 		}
 		else
@@ -65,10 +74,16 @@ void UObjectMixerEditorSerializedData::AddObjectsToCollection(const FName& Filte
 		}
 
 		SaveConfig();
+
+		OnObjectMixerCollectionMapChanged.Broadcast();
+
+		return true;
 	}
+
+	return false;
 }
 
-void UObjectMixerEditorSerializedData::RemoveObjectsFromCollection(const FName& FilterClassName, const FName& CollectionName, const TSet<FSoftObjectPath>& ObjectsToRemove)
+bool UObjectMixerEditorSerializedData::RemoveObjectsFromCollection(const FName& FilterClassName, const FName& CollectionName, const TSet<FSoftObjectPath>& ObjectsToRemove)
 {
 	if (FObjectMixerSerializationDataPerFilter* Data = FindSerializationDataByFilterClassName(FilterClassName))
 	{		
@@ -90,12 +105,18 @@ void UObjectMixerEditorSerializedData::RemoveObjectsFromCollection(const FName&
 			}));
 			
 			SaveConfig();
+
+			OnObjectMixerCollectionMapChanged.Broadcast();
+
+			return true;
 		}
 		
 	}
+
+	return false;
 }
 
-void UObjectMixerEditorSerializedData::RemoveCollection(const FName& FilterClassName, const FName& CollectionName)
+bool UObjectMixerEditorSerializedData::RemoveCollection(const FName& FilterClassName, const FName& CollectionName)
 {	
 	if (FObjectMixerSerializationDataPerFilter* Data = FindSerializationDataByFilterClassName(FilterClassName))
 	{
@@ -110,11 +131,64 @@ void UObjectMixerEditorSerializedData::RemoveCollection(const FName& FilterClass
 			}));
 		
 		SaveConfig();
+
+		OnObjectMixerCollectionMapChanged.Broadcast();
+
+		return true;
+	}
+
+	return false;
+}
+
+bool UObjectMixerEditorSerializedData::DuplicateCollection(const FName& FilterClassName,
+	const FName& CollectionToDuplicateName, FName& DesiredDuplicateName)
+{
+	uint32 NumberToAppendToCollectionName = 0;
+
+	FName DuplicateName = FName(DesiredDuplicateName.ToString() + FString::FromInt(NumberToAppendToCollectionName));
+	
+	while (DoesCollectionExist(FilterClassName, DuplicateName))
+	{
+		NumberToAppendToCollectionName++;
+
+		DuplicateName = FName(DesiredDuplicateName.ToString() + FString::FromInt(NumberToAppendToCollectionName));
+	}
+
+	DesiredDuplicateName = DuplicateName;
+
+	if (FObjectMixerSerializationDataPerFilter* Data = FindSerializationDataByFilterClassName(FilterClassName))
+	{
+		FObjectMixerCollectionObjectSet MatchCopy;
+		if (const FObjectMixerCollectionObjectSet* Match = Algo::FindByPredicate(Data->SerializedCollections,
+			[CollectionToDuplicateName](const FObjectMixerCollectionObjectSet& Comparator)
+			{
+				return Comparator.CollectionName.IsEqual(CollectionToDuplicateName);
+			}))
+		{
+			
+			FScopedTransaction DuplicateCollectionTransaction(LOCTEXT("DuplicateCollectionTransaction","Duplicate Collection"));
+	
+			Modify();
+			
+			MatchCopy = (*Match);
+
+			MatchCopy.CollectionName = DuplicateName;
+
+			Data->SerializedCollections.Add(MatchCopy);
+		
+			SaveConfig();
+
+			OnObjectMixerCollectionMapChanged.Broadcast();
+
+			return true;
+		}
 	}
+
+	return false;
 }
 
-void UObjectMixerEditorSerializedData::ReorderCollection(const FName& FilterClassName,
-	const FName& CollectionToMoveName, const FName& CollectionInsertBeforeName)
+bool UObjectMixerEditorSerializedData::ReorderCollection(const FName& FilterClassName,
+                                                         const FName& CollectionToMoveName, const FName& CollectionInsertBeforeName)
 {
 	if (FObjectMixerSerializationDataPerFilter* Data = FindSerializationDataByFilterClassName(FilterClassName))
 	{
@@ -138,7 +212,7 @@ void UObjectMixerEditorSerializedData::ReorderCollection(const FName& FilterClas
 			
 			RemoveCollection(FilterClassName, CollectionToMoveName);
 
-			if (CollectionInsertBeforeName == "All") // Move MatchCopy to the end
+			if (CollectionInsertBeforeName == UObjectMixerEditorSerializedData::AllCollectionName) // Move MatchCopy to the end
 			{
 				Data->SerializedCollections.Add(MatchCopy);
 			}
@@ -154,8 +228,66 @@ void UObjectMixerEditorSerializedData::ReorderCollection(const FName& FilterClas
 			}
 		
 			SaveConfig();
+
+			OnObjectMixerCollectionMapChanged.Broadcast();
+
+			return true;
 		}
 	}
+
+	return false;
+}
+
+bool UObjectMixerEditorSerializedData::RenameCollection(const FName& FilterClassName,
+	const FName CollectionNameToRename, const FName NewCollectionName)
+{
+	if (NewCollectionName == NAME_None ||
+		NewCollectionName.IsEqual(UObjectMixerEditorSerializedData::AllCollectionName, ENameCase::IgnoreCase) ||
+		NewCollectionName.ToString().TrimStartAndEnd().IsEmpty() ||
+		DoesCollectionExist(FilterClassName, NewCollectionName))
+	{
+		return false;
+	}
+		
+	if (FObjectMixerSerializationDataPerFilter* Data = FindSerializationDataByFilterClassName(FilterClassName))
+	{
+		if (FObjectMixerCollectionObjectSet* Match = Algo::FindByPredicate(Data->SerializedCollections,
+			[CollectionNameToRename](const FObjectMixerCollectionObjectSet& Comparator)
+			{
+				return Comparator.CollectionName.IsEqual(CollectionNameToRename);
+			}))
+		{
+			FScopedTransaction RenameCollectionTransaction(LOCTEXT("RenameCollectionTransaction","Rename Collection"));
+
+			Modify();
+
+			Match->CollectionName = NewCollectionName;
+	
+			SaveConfig();
+
+			OnObjectMixerCollectionRenamed.Broadcast(CollectionNameToRename, NewCollectionName);
+
+			OnObjectMixerCollectionMapChanged.Broadcast();
+
+			return true;
+		}
+	}
+
+	return false;
+}
+
+bool UObjectMixerEditorSerializedData::DoesCollectionExist(const FName& FilterClassName, const FName& CollectionName)
+{
+	if (FObjectMixerSerializationDataPerFilter* Data = FindSerializationDataByFilterClassName(FilterClassName))
+	{
+		return Algo::FindByPredicate(Data->SerializedCollections,
+			[CollectionName](const FObjectMixerCollectionObjectSet& Comparator)
+			{
+				return Comparator.CollectionName.IsEqual(CollectionName);
+			}) != nullptr;
+	}
+
+	return false;
 }
 
 bool UObjectMixerEditorSerializedData::IsObjectInCollection(const FName& FilterClassName, const FName& CollectionName, const FSoftObjectPath& InObject)
@@ -250,14 +382,11 @@ bool UObjectMixerEditorSerializedData::IsColumnDataSerialized(const FName& Filte
 {
 	if (FObjectMixerSerializationDataPerFilter* Data = FindSerializationDataByFilterClassName(FilterClassName))
 	{
-		if (const FObjectMixerColumnData* Match = Algo::FindByPredicate(Data->SerializedColumnData,
+		return Algo::FindByPredicate(Data->SerializedColumnData,
 			[ColumnName](const FObjectMixerColumnData& Comparator)
 			{
 				return Comparator.ColumnName.IsEqual(ColumnName);
-			}))
-		{
-			return true;
-		}
+			}) != nullptr;
 	}
 
 	return false;
diff --git a/Engine/Plugins/Editor/ObjectMixer/ObjectMixer/Source/ObjectMixer/Private/Views/MainPanel/ObjectMixerEditorMainPanel.cpp b/Engine/Plugins/Editor/ObjectMixer/ObjectMixer/Source/ObjectMixer/Private/Views/MainPanel/ObjectMixerEditorMainPanel.cpp
--- a/Engine/Plugins/Editor/ObjectMixer/ObjectMixer/Source/ObjectMixer/Private/Views/MainPanel/ObjectMixerEditorMainPanel.cpp
+++ b/Engine/Plugins/Editor/ObjectMixer/ObjectMixer/Source/ObjectMixer/Private/Views/MainPanel/ObjectMixerEditorMainPanel.cpp
@@ -110,66 +110,101 @@ const TArray<TSharedRef<IObjectMixerEditorListFilter>>& FObjectMixerEditorMainPa
 	return MainPanelWidget->GetShowFilters();	
 }
 
-void FObjectMixerEditorMainPanel::AddObjectsToCollection(const FName& CollectionName, const TSet<FSoftObjectPath>& ObjectsToAdd) const
+UObjectMixerEditorSerializedData* FObjectMixerEditorMainPanel::GetSerializedDataOutputtingFilterName(FName& OutFilterName) const
 {
 	if (const TSubclassOf<UObjectMixerObjectFilter> Filter = GetObjectFilterClass())
 	{
-		if (UObjectMixerEditorSerializedData* SerializedData = GetMutableDefault<UObjectMixerEditorSerializedData>())
-		{
-			SerializedData->AddObjectsToCollection(Filter->GetFName(), CollectionName, ObjectsToAdd);
+		OutFilterName = Filter->GetFName();
+		return GetMutableDefault<UObjectMixerEditorSerializedData>();
+	}
+
+	return nullptr;
+}
 
-			OnObjectMixerCollectionMapChanged.Broadcast();
-		}
+bool FObjectMixerEditorMainPanel::RequestAddObjectsToCollection(const FName& CollectionName, const TSet<FSoftObjectPath>& ObjectsToAdd) const
+{
+	FName FilterName;
+	if (UObjectMixerEditorSerializedData* SerializedData = GetSerializedDataOutputtingFilterName(FilterName))
+	{
+		return SerializedData->AddObjectsToCollection(FilterName, CollectionName, ObjectsToAdd);
 	}
+
+	return false;
 }
 
-void FObjectMixerEditorMainPanel::RemoveObjectsFromCollection(const FName& CollectionName, const TSet<FSoftObjectPath>& ObjectsToRemove) const
+bool FObjectMixerEditorMainPanel::RequestRemoveObjectsFromCollection(const FName& CollectionName, const TSet<FSoftObjectPath>& ObjectsToRemove) const
 {
-	if (const TSubclassOf<UObjectMixerObjectFilter> Filter = GetObjectFilterClass())
+	FName FilterName;
+	if (UObjectMixerEditorSerializedData* SerializedData = GetSerializedDataOutputtingFilterName(FilterName))
 	{
-		if (UObjectMixerEditorSerializedData* SerializedData = GetMutableDefault<UObjectMixerEditorSerializedData>())
-		{
-			SerializedData->RemoveObjectsFromCollection(Filter->GetFName(), CollectionName, ObjectsToRemove);
+		return SerializedData->RemoveObjectsFromCollection(FilterName, CollectionName, ObjectsToRemove);
+	}
 
-			OnObjectMixerCollectionMapChanged.Broadcast();
-		}
+	return false;
+}
+
+bool FObjectMixerEditorMainPanel::RequestRemoveCollection(const FName& CollectionName) const
+{
+	FName FilterName;
+	if (UObjectMixerEditorSerializedData* SerializedData = GetSerializedDataOutputtingFilterName(FilterName))
+	{
+		return SerializedData->RemoveCollection(FilterName, CollectionName);
 	}
+
+	return false;
 }
 
-void FObjectMixerEditorMainPanel::RemoveCollection(const FName& CollectionName) const
+bool FObjectMixerEditorMainPanel::RequestDuplicateCollection(const FName& CollectionToDuplicateName, FName& DesiredDuplicateName) const
 {
-	if (const TSubclassOf<UObjectMixerObjectFilter> Filter = GetObjectFilterClass())
+	FName FilterName;
+	if (UObjectMixerEditorSerializedData* SerializedData = GetSerializedDataOutputtingFilterName(FilterName))
 	{
-		if (UObjectMixerEditorSerializedData* SerializedData = GetMutableDefault<UObjectMixerEditorSerializedData>())
-		{
-			SerializedData->RemoveCollection(Filter->GetFName(), CollectionName);
+		return SerializedData->DuplicateCollection(FilterName, CollectionToDuplicateName, DesiredDuplicateName);
+	}
 
-			OnObjectMixerCollectionMapChanged.Broadcast();
-		}
+	return false;
+}
+
+bool FObjectMixerEditorMainPanel::RequestReorderCollection(const FName& CollectionToMoveName, const FName& CollectionInsertBeforeName) const
+{
+	FName FilterName;
+	if (UObjectMixerEditorSerializedData* SerializedData = GetSerializedDataOutputtingFilterName(FilterName))
+	{
+		return SerializedData->ReorderCollection(FilterName, CollectionToMoveName, CollectionInsertBeforeName);
 	}
+
+	return false;
 }
 
-void FObjectMixerEditorMainPanel::ReorderCollection(const FName& CollectionToMoveName, const FName& CollectionInsertBeforeName) const
+bool FObjectMixerEditorMainPanel::RequestRenameCollection(const FName& CollectionNameToRename,
+	const FName& NewCollectionName) const
 {
-	if (const TSubclassOf<UObjectMixerObjectFilter> Filter = GetObjectFilterClass())
+	FName FilterName;
+	if (UObjectMixerEditorSerializedData* SerializedData = GetSerializedDataOutputtingFilterName(FilterName))
 	{
-		if (UObjectMixerEditorSerializedData* SerializedData = GetMutableDefault<UObjectMixerEditorSerializedData>())
-		{
-			SerializedData->ReorderCollection(Filter->GetFName(), CollectionToMoveName, CollectionInsertBeforeName);
+		return SerializedData->RenameCollection(FilterName, CollectionNameToRename, NewCollectionName);
+	}
+
+	return false;
+}
 
-			OnObjectMixerCollectionMapChanged.Broadcast();
-		}
+bool FObjectMixerEditorMainPanel::DoesCollectionExist(const FName& CollectionName) const
+{
+	FName FilterName;
+	if (UObjectMixerEditorSerializedData* SerializedData = GetSerializedDataOutputtingFilterName(FilterName))
+	{
+		return SerializedData->DoesCollectionExist(FilterName, CollectionName);
 	}
+
+	return false;
 }
 
 bool FObjectMixerEditorMainPanel::IsObjectInCollection(const FName& CollectionName, const FSoftObjectPath& InObject) const
 {
-	if (const TSubclassOf<UObjectMixerObjectFilter> Filter = GetObjectFilterClass())
+	FName FilterName;
+	if (UObjectMixerEditorSerializedData* SerializedData = GetSerializedDataOutputtingFilterName(FilterName))
 	{
-		if (UObjectMixerEditorSerializedData* SerializedData = GetMutableDefault<UObjectMixerEditorSerializedData>())
-		{
-			return SerializedData->IsObjectInCollection(Filter->GetFName(), CollectionName, InObject);
-		}
+		return SerializedData->IsObjectInCollection(FilterName, CollectionName, InObject);
 	}
 
 	return false;
@@ -177,12 +212,10 @@ bool FObjectMixerEditorMainPanel::IsObjectInCollection(const FName& CollectionNa
 
 TSet<FName> FObjectMixerEditorMainPanel::GetCollectionsForObject(const FSoftObjectPath& InObject) const
 {
-	if (const TSubclassOf<UObjectMixerObjectFilter> Filter = GetObjectFilterClass())
+	FName FilterName;
+	if (UObjectMixerEditorSerializedData* SerializedData = GetSerializedDataOutputtingFilterName(FilterName))
 	{
-		if (UObjectMixerEditorSerializedData* SerializedData = GetMutableDefault<UObjectMixerEditorSerializedData>())
-		{
-			return SerializedData->GetCollectionsForObject(Filter->GetFName(), InObject);
-		}
+		return SerializedData->GetCollectionsForObject(FilterName, InObject);
 	}
 
 	return {};
@@ -190,12 +223,10 @@ TSet<FName> FObjectMixerEditorMainPanel::GetCollectionsForObject(const FSoftObje
 
 TArray<FName> FObjectMixerEditorMainPanel::GetAllCollectionNames() const
 {
-	if (const TSubclassOf<UObjectMixerObjectFilter> Filter = GetObjectFilterClass())
+	FName FilterName;
+	if (UObjectMixerEditorSerializedData* SerializedData = GetSerializedDataOutputtingFilterName(FilterName))
 	{
-		if (UObjectMixerEditorSerializedData* SerializedData = GetMutableDefault<UObjectMixerEditorSerializedData>())
-		{
-			return SerializedData->GetAllCollectionNames(Filter->GetFName());
-		}
+		return SerializedData->GetAllCollectionNames(FilterName);
 	}
 
 	return {};
diff --git a/Engine/Plugins/Editor/ObjectMixer/ObjectMixer/Source/ObjectMixer/Private/Views/MainPanel/SObjectMixerEditorMainPanel.cpp b/Engine/Plugins/Editor/ObjectMixer/ObjectMixer/Source/ObjectMixer/Private/Views/MainPanel/SObjectMixerEditorMainPanel.cpp
--- a/Engine/Plugins/Editor/ObjectMixer/ObjectMixer/Source/ObjectMixer/Private/Views/MainPanel/SObjectMixerEditorMainPanel.cpp
+++ b/Engine/Plugins/Editor/ObjectMixer/ObjectMixer/Source/ObjectMixer/Private/Views/MainPanel/SObjectMixerEditorMainPanel.cpp
@@ -4,8 +4,10 @@
 
 #include "ObjectMixerEditorLog.h"
 #include "ObjectMixerEditorModule.h"
+#include "ObjectMixerEditorSerializedData.h"
 #include "ObjectMixerEditorSettings.h"
 #include "Views/List/ObjectMixerEditorList.h"
+#include "Views/List/ObjectMixerEditorListFilters/ObjectMixerEditorListFilter_Collection.h"
 #include "Views/MainPanel/ObjectMixerEditorMainPanel.h"
 #include "Views/Widgets/SCollectionSelectionButton.h"
 #include "Views/Widgets/SObjectMixerPlacementAssetMenuEntry.h"
@@ -17,8 +19,6 @@
 #include "PlacementMode/Public/IPlacementModeModule.h"
 #include "SPositiveActionButton.h"
 #include "Styling/StyleColors.h"
-#include "Views/List/ObjectMixerEditorListFilters/IObjectMixerEditorListFilter.h"
-#include "Views/List/ObjectMixerEditorListFilters/ObjectMixerEditorListFilter_Collection.h"
 #include "Widgets/Input/SCheckBox.h"
 #include "Widgets/Input/SComboButton.h"
 #include "Widgets/Input/SMultiLineEditableTextBox.h"
@@ -64,7 +64,10 @@ void SObjectMixerEditorMainPanel::Construct(
 		]
 	];
 
-	InMainPanel->GetOnObjectMixerCollectionMapChanged().AddRaw(this, &SObjectMixerEditorMainPanel::RebuildCollectionSelector);
+	if (UObjectMixerEditorSerializedData* SerializedData = GetMutableDefault<UObjectMixerEditorSerializedData>())
+	{
+		SerializedData->OnObjectMixerCollectionMapChanged.AddRaw(this, &SObjectMixerEditorMainPanel::RebuildCollectionSelector);
+	}
 
 	ShowFilters.Add(MakeShared<FObjectMixerEditorListFilter_Collection>());
 }
@@ -540,7 +543,7 @@ void SObjectMixerEditorMainPanel::RebuildCollectionSelector()
 
 	CollectionSelectorBox->AddSlot()
 	[
-		CreateSection("All")
+		CreateSection(UObjectMixerEditorSerializedData::AllCollectionName)
 	];
 
 	for (const FName& Key : AllCollections)
@@ -554,11 +557,44 @@ void SObjectMixerEditorMainPanel::RebuildCollectionSelector()
 	CollectionSelectorBox->SetVisibility(EVisibility::Visible);
 }
 
-void SObjectMixerEditorMainPanel::RequestRemoveCollection(const FName& CollectionName)
+bool SObjectMixerEditorMainPanel::RequestRemoveCollection(const FName& CollectionName)
+{
+	if (MainPanelModel.Pin()->RequestRemoveCollection(CollectionName))
+	{
+		CurrentCollectionSelection.Remove(CollectionName);
+
+		return true;
+	}
+
+	return false;
+}
+
+bool SObjectMixerEditorMainPanel::RequestDuplicateCollection(const FName& CollectionToDuplicateName, FName& DesiredDuplicateName) const
+{
+	return MainPanelModel.Pin()->RequestDuplicateCollection(CollectionToDuplicateName, DesiredDuplicateName);
+}
+
+bool SObjectMixerEditorMainPanel::RequestRenameCollection(
+	const FName& CollectionNameToRename,
+	const FName& NewCollectionName)
 {
-	MainPanelModel.Pin()->RemoveCollection(CollectionName);
+	if (MainPanelModel.Pin()->RequestRenameCollection(CollectionNameToRename, NewCollectionName))
+	{
+		if (CurrentCollectionSelection.Contains(CollectionNameToRename))
+		{
+			CurrentCollectionSelection.Remove(CollectionNameToRename);
+			CurrentCollectionSelection.Add(NewCollectionName);
+		}
+
+		return true;
+	}
+
+	return false;
+}
 
-	CurrentCollectionSelection.Remove(CollectionName);
+bool SObjectMixerEditorMainPanel::DoesCollectionExist(const FName& CollectionName) const
+{
+	return MainPanelModel.Pin()->DoesCollectionExist(CollectionName);
 }
 
 void SObjectMixerEditorMainPanel::OnCollectionCheckedStateChanged(ECheckBoxState State, FName CollectionName)
@@ -584,7 +620,7 @@ void SObjectMixerEditorMainPanel::OnCollectionCheckedStateChanged(ECheckBoxState
 			{
 				CurrentCollectionSelection.Reset();
 
-				if (CollectionName != "All")
+				if (CollectionName != UObjectMixerEditorSerializedData::AllCollectionName && CollectionName != NAME_None)
 				{
 					CurrentCollectionSelection.Add(CollectionName);
 				}
@@ -597,7 +633,7 @@ void SObjectMixerEditorMainPanel::OnCollectionCheckedStateChanged(ECheckBoxState
 				CurrentCollectionSelection.Reset();
 			}
 
-			if (CollectionName != "All")
+			if (CollectionName != UObjectMixerEditorSerializedData::AllCollectionName && CollectionName != NAME_None)
 			{
 				CurrentCollectionSelection.Add(CollectionName);
 			}
@@ -609,7 +645,7 @@ void SObjectMixerEditorMainPanel::OnCollectionCheckedStateChanged(ECheckBoxState
 
 ECheckBoxState SObjectMixerEditorMainPanel::IsCollectionChecked(FName Section) const
 {
-	if (CurrentCollectionSelection.IsEmpty() && Section == "All")
+	if (CurrentCollectionSelection.IsEmpty() && Section == UObjectMixerEditorSerializedData::AllCollectionName)
 	{
 		return ECheckBoxState::Checked;
 	}
@@ -619,11 +655,7 @@ ECheckBoxState SObjectMixerEditorMainPanel::IsCollectionChecked(FName Section) c
 
 SObjectMixerEditorMainPanel::~SObjectMixerEditorMainPanel()
 {
-	if (MainPanelModel.IsValid())
-	{
-		MainPanelModel.Pin()->GetOnObjectMixerCollectionMapChanged().RemoveAll(this);
-		MainPanelModel.Reset();
-	}
+	MainPanelModel.Reset();
 }
 
 #undef LOCTEXT_NAMESPACE
diff --git a/Engine/Plugins/Editor/ObjectMixer/ObjectMixer/Source/ObjectMixer/Private/Views/Widgets/ObjectMixerEditorListMenuContext.cpp b/Engine/Plugins/Editor/ObjectMixer/ObjectMixer/Source/ObjectMixer/Private/Views/Widgets/ObjectMixerEditorListMenuContext.cpp
--- a/Engine/Plugins/Editor/ObjectMixer/ObjectMixer/Source/ObjectMixer/Private/Views/Widgets/ObjectMixerEditorListMenuContext.cpp
+++ b/Engine/Plugins/Editor/ObjectMixer/ObjectMixer/Source/ObjectMixer/Private/Views/Widgets/ObjectMixerEditorListMenuContext.cpp
@@ -2,6 +2,8 @@
 
 #include "Views/Widgets/ObjectMixerEditorListMenuContext.h"
 
+#include "Views/MainPanel/ObjectMixerEditorMainPanel.h"
+
 #include "Elements/Framework/TypedElementSelectionSet.h"
 #include "Elements/Interfaces/TypedElementObjectInterface.h"
 #include "Framework/MultiBox/MultiBoxDefs.h"
@@ -16,7 +18,6 @@
 #include "ToolMenuDelegates.h"
 #include "ToolMenuEntry.h"
 #include "ToolMenuSection.h"
-#include "Views/MainPanel/ObjectMixerEditorMainPanel.h"
 #include "Widgets/Layout/SBox.h"
 #include "Widgets/Input/SEditableTextBox.h"
 
@@ -142,7 +143,6 @@ void UObjectMixerEditorListMenuContext::RegisterObjectMixerContextMenu()
 												SNew(SEditableTextBox)
 												.HintText(LOCTEXT("NewCollectionEditableHintText", "Enter a new collection name..."))
 												.OnTextCommitted_Static(&UObjectMixerEditorListMenuContext::OnTextCommitted, ContextData)
-												.OnTextChanged_Static(&UObjectMixerEditorListMenuContext::OnTextChanged, ContextData)
 											]
 									;
 								}
@@ -193,20 +193,10 @@ void UObjectMixerEditorListMenuContext::OnTextCommitted(const FText& InText, ETe
 {
 	if (InCommitType == ETextCommit::OnEnter)
 	{
-		UE_LOG(LogTemp, Warning, TEXT("%hs: %s"), __FUNCTION__, *InText.ToString());
-
 		AddObjectsToCollection(*InText.ToString(), ContextData);
 	}
 }
 
-void UObjectMixerEditorListMenuContext::OnTextChanged(const FText& InText, const FObjectMixerEditorListMenuContextData ContextData)
-{
-	if (const FString AsString = InText.ToString(); AsString.Len() > 12)
-	{
-		//ContextData.EditableText->SetText(FText::FromString(AsString.LeftChop(1)));
-	}
-}
-
 void UObjectMixerEditorListMenuContext::OnClickCollectionMenuEntry(const FName Key, const FObjectMixerEditorListMenuContextData ContextData)
 {
 	if (AreAllObjectsInCollection(Key, ContextData))
@@ -233,7 +223,7 @@ void UObjectMixerEditorListMenuContext::AddObjectsToCollection(const FName Key,
 			}
 		}
 		
-		MainPanel->AddObjectsToCollection(Key, ObjectPaths);
+		MainPanel->RequestAddObjectsToCollection(Key, ObjectPaths);
 	}
 }
 
@@ -251,7 +241,7 @@ void UObjectMixerEditorListMenuContext::RemoveObjectsFromCollection(const FName
 			}
 		}
 		
-		MainPanel->RemoveObjectsFromCollection(Key, ObjectPaths);
+		MainPanel->RequestRemoveObjectsFromCollection(Key, ObjectPaths);
 	}
 }
 
diff --git a/Engine/Plugins/Editor/ObjectMixer/ObjectMixer/Source/ObjectMixer/Private/Views/Widgets/SCollectionSelectionButton.cpp b/Engine/Plugins/Editor/ObjectMixer/ObjectMixer/Source/ObjectMixer/Private/Views/Widgets/SCollectionSelectionButton.cpp
--- a/Engine/Plugins/Editor/ObjectMixer/ObjectMixer/Source/ObjectMixer/Private/Views/Widgets/SCollectionSelectionButton.cpp
+++ b/Engine/Plugins/Editor/ObjectMixer/ObjectMixer/Source/ObjectMixer/Private/Views/Widgets/SCollectionSelectionButton.cpp
@@ -2,12 +2,18 @@
 
 #include "Views/Widgets/SCollectionSelectionButton.h"
 
-#include "DragAndDrop/DecoratedDragDropOp.h"
-#include "Styling/AppStyle.h"
+#include "ObjectMixerEditorSerializedData.h"
 #include "Views/MainPanel/ObjectMixerEditorMainPanel.h"
 #include "Views/MainPanel/SObjectMixerEditorMainPanel.h"
+
+#include "DragAndDrop/DecoratedDragDropOp.h"
+#include "Framework/Commands/GenericCommands.h"
+#include "Framework/MultiBox/MultiBoxBuilder.h"
+#include "Styling/AppStyle.h"
+#include "Views/List/ObjectMixerEditorListRow.h"
+#include "Widgets/Input/SMenuAnchor.h"
 #include "Widgets/Layout/SBorder.h"
-#include "Widgets/Text/STextBlock.h"
+#include "Widgets/Text/SInlineEditableTextBlock.h"
 
 #define LOCTEXT_NAMESPACE "ObjectMixerEditor"
 
@@ -46,20 +52,75 @@ void SCollectionSelectionButton::Construct(const FArguments& InArgs, const TShar
 	
 	ChildSlot
 	[
-		SNew(SBorder)
-		.Padding(FMargin(16, 4))
-		.BorderImage(this, &SCollectionSelectionButton::GetBorderBrush)
-		.ForegroundColor(this, &SCollectionSelectionButton::GetBorderForeground)
+		SAssignNew(MenuAnchor, SMenuAnchor)
+		.OnGetMenuContent(this, &SCollectionSelectionButton::GetContextMenu)
 		[
-			SNew(STextBlock)
-			.TextStyle(FAppStyle::Get(), "SmallText")
-			.Text(FText::FromName(CollectionName))
-			.Visibility(EVisibility::HitTestInvisible)
-			.ColorAndOpacity(FSlateColor::UseForeground())
+			SNew(SBorder)
+			.Padding(FMargin(16, 4))
+			.BorderImage(this, &SCollectionSelectionButton::GetBorderBrush)
+			.ForegroundColor(this, &SCollectionSelectionButton::GetBorderForeground)
+			[
+				SAssignNew(EditableTextBlock, SInlineEditableTextBlock)
+				.Style(&FAppStyle::Get().GetWidgetStyle<FInlineEditableTextBlockStyle>("InlineEditableTextBlockSmallStyle"))
+				.Text(FText::FromName(CollectionName))
+				.Visibility(EVisibility::HitTestInvisible)
+				.ColorAndOpacity(FSlateColor::UseForeground())
+				.OnTextCommitted(this, &SCollectionSelectionButton::OnEditableTextCommitted)
+			]
 		]
 	];
 }
 
+TSharedRef<SWidget> SCollectionSelectionButton::GetContextMenu() const
+{
+	FMenuBuilder MenuBuilder(true, nullptr);
+
+	MenuBuilder.AddMenuEntry(
+		FText::Format(LOCTEXT("DeleteCollectionButtonMenuEntry", "Delete '{0}'"), FText::FromName(CollectionName)),
+		LOCTEXT("DeleteCollectionButtonMenuEntryTooltip", "Remove this collection. This operation can be undone."),
+		FGenericCommands::Get().Delete->GetIcon(),
+		FUIAction(FExecuteAction::CreateLambda([this]()
+		{
+			MainPanelPtr.Pin()->RequestRemoveCollection(CollectionName);
+		}))
+	);
+
+	MenuBuilder.AddMenuEntry(
+		FText::Format(LOCTEXT("DuplicateCollectionButtonMenuEntry", "Duplicate '{0}'"), FText::FromName(CollectionName)),
+		LOCTEXT("DuplicateCollectionButtonMenuEntryTooltip", "Duplicate this collection. This operation can be undone."),
+		FGenericCommands::Get().Duplicate->GetIcon(),
+		FUIAction(FExecuteAction::CreateLambda([this]()
+		{
+			FName DuplicateName = CollectionName;
+			MainPanelPtr.Pin()->RequestDuplicateCollection(CollectionName, DuplicateName);
+		}))
+	);
+
+	MenuBuilder.AddMenuEntry(
+		FText::Format(LOCTEXT("RenameCollectionButtonMenuEntry", "Rename '{0}'"), FText::FromName(CollectionName)),
+		LOCTEXT("RenameCollectionButtonMenuEntryTooltip", "Rename this collection. This operation can be undone."),
+		FGenericCommands::Get().Rename->GetIcon(),
+		FUIAction(FExecuteAction::CreateLambda([this]()
+		{
+			EditableTextBlock->EnterEditingMode();
+		}))
+	);
+			
+	return MenuBuilder.MakeWidget();
+}
+
+void SCollectionSelectionButton::OnEditableTextCommitted(const FText& Text, ETextCommit::Type Type)
+{
+	if (Type == ETextCommit::OnEnter)
+	{
+		const FName NewCollectionName = *Text.ToString();
+		if (MainPanelPtr.Pin()->RequestRenameCollection(CollectionName, NewCollectionName))
+		{
+			CollectionName = NewCollectionName;
+		}
+	}
+}
+
 FReply SCollectionSelectionButton::OnMouseButtonDown(const FGeometry& MyGeometry, const FPointerEvent& MouseEvent) 
 {
 	bIsPressed = true;
@@ -93,14 +154,17 @@ FReply SCollectionSelectionButton::OnMouseButtonUp(const FGeometry& MyGeometry,
 	}
 	if (MouseEvent.GetEffectingButton() == EKeys::RightMouseButton)
 	{
-		MainPanelPtr.Pin()->RequestRemoveCollection(CollectionName);
+		if (CollectionName != UObjectMixerEditorSerializedData::AllCollectionName && CollectionName != NAME_None)
+		{
+			MenuAnchor->SetIsOpen(true);
+		}
 	}
 	return FReply::Handled();
 }
 
 FReply SCollectionSelectionButton::OnDragDetected(const FGeometry& MyGeometry, const FPointerEvent& MouseEvent)
 {
-	if (CollectionName != "All")
+	if (CollectionName != UObjectMixerEditorSerializedData::AllCollectionName && CollectionName != NAME_None)
 	{
 		TSharedRef<FCollectionSelectionButtonDragDropOp> OperationFromCollection =
 		   FCollectionSelectionButtonDragDropOp::New(CollectionName);
@@ -120,7 +184,7 @@ void SCollectionSelectionButton::OnDragEnter(const FGeometry& MyGeometry, const
 	if (TSharedPtr<FObjectMixerListRowDragDropOp> OperationFromRow =
 		DragDropEvent.GetOperationAs<FObjectMixerListRowDragDropOp>())
 	{
-		if (CollectionName != "All" && CollectionName != NAME_None)
+		if (CollectionName != UObjectMixerEditorSerializedData::AllCollectionName && CollectionName != NAME_None)
 		{
 			OperationFromRow->SetToolTip(
 			   LOCTEXT("DropRowItemsOntoCollectionButtonCTA","Add selected items to this collection"),
@@ -133,7 +197,7 @@ void SCollectionSelectionButton::OnDragEnter(const FGeometry& MyGeometry, const
 	if (TSharedPtr<FCollectionSelectionButtonDragDropOp> OperationFromCollection =
 		DragDropEvent.GetOperationAs<FCollectionSelectionButtonDragDropOp>())
 	{
-		if (CollectionName != "All" && CollectionName != NAME_None && CollectionName != OperationFromCollection->DraggedItem)
+		if (CollectionName != UObjectMixerEditorSerializedData::AllCollectionName && CollectionName != NAME_None && CollectionName != OperationFromCollection->DraggedItem)
 		{
 			OperationFromCollection->SetToolTip(
 			   FText::Format(
@@ -172,7 +236,7 @@ FReply SCollectionSelectionButton::OnDrop(const FGeometry& MyGeometry, const FDr
 	if (TSharedPtr<FObjectMixerListRowDragDropOp> OperationFromRow =
 		DragDropEvent.GetOperationAs<FObjectMixerListRowDragDropOp>())
 	{
-		if (CollectionName != "All" && CollectionName != NAME_None)
+		if (CollectionName != UObjectMixerEditorSerializedData::AllCollectionName && CollectionName != NAME_None)
 		{
 			if (const TSharedPtr<FObjectMixerEditorMainPanel> MainPanelModel = MainPanelPtr.Pin()->GetMainPanelModel().Pin())
 			{
@@ -186,7 +250,7 @@ FReply SCollectionSelectionButton::OnDrop(const FGeometry& MyGeometry, const FDr
 					}
 				}
 	
-				MainPanelModel->AddObjectsToCollection(CollectionName, ObjectPaths);
+				MainPanelModel->RequestAddObjectsToCollection(CollectionName, ObjectPaths);
 			}
 		}
 	}
@@ -194,11 +258,11 @@ FReply SCollectionSelectionButton::OnDrop(const FGeometry& MyGeometry, const FDr
 	if (TSharedPtr<FCollectionSelectionButtonDragDropOp> OperationFromCollection =
 		DragDropEvent.GetOperationAs<FCollectionSelectionButtonDragDropOp>())
 	{
-		if (CollectionName != "All" && CollectionName != NAME_None && CollectionName != OperationFromCollection->DraggedItem)
+		if (CollectionName != UObjectMixerEditorSerializedData::AllCollectionName && CollectionName != NAME_None && CollectionName != OperationFromCollection->DraggedItem)
 		{
 			if (const TSharedPtr<FObjectMixerEditorMainPanel> MainPanelModel = MainPanelPtr.Pin()->GetMainPanelModel().Pin())
 			{
-				MainPanelModel->ReorderCollection(OperationFromCollection->DraggedItem, CollectionName);
+				MainPanelModel->RequestReorderCollection(OperationFromCollection->DraggedItem, CollectionName);
 			}
 		}
 	}
diff --git a/Engine/Plugins/Editor/ObjectMixer/ObjectMixer/Source/ObjectMixer/Public/ObjectMixerEditorSerializedData.h b/Engine/Plugins/Editor/ObjectMixer/ObjectMixer/Source/ObjectMixer/Public/ObjectMixerEditorSerializedData.h
--- a/Engine/Plugins/Editor/ObjectMixer/ObjectMixer/Source/ObjectMixer/Public/ObjectMixerEditorSerializedData.h
+++ b/Engine/Plugins/Editor/ObjectMixer/ObjectMixer/Source/ObjectMixer/Public/ObjectMixerEditorSerializedData.h
@@ -4,6 +4,9 @@
 
 #include "ObjectMixerEditorSerializedData.generated.h"
 
+DECLARE_MULTICAST_DELEGATE(FOnObjectMixerCollectionChanged)
+DECLARE_MULTICAST_DELEGATE_TwoParams(FOnObjectMixerCollectionRenamed, FName, FName)
+
 USTRUCT()
 struct FObjectMixerCollectionObjectData
 {
@@ -110,13 +113,19 @@ public:
 
 	FObjectMixerSerializationDataPerFilter* FindSerializationDataByFilterClassName(const FName& FilterClassName);
 
-	void AddObjectsToCollection(const FName& FilterClassName, const FName& CollectionName, const TSet<FSoftObjectPath>& ObjectsToAdd);
+	bool AddObjectsToCollection(const FName& FilterClassName, const FName& CollectionName, const TSet<FSoftObjectPath>& ObjectsToAdd);
+
+	bool RemoveObjectsFromCollection(const FName& FilterClassName, const FName& CollectionName, const TSet<FSoftObjectPath>& ObjectsToRemove);
+
+	bool RemoveCollection(const FName& FilterClassName, const FName& CollectionName);
+	
+	bool DuplicateCollection(const FName& FilterClassName, const FName& CollectionToDuplicateName, FName& DesiredDuplicateName);
 
-	void RemoveObjectsFromCollection(const FName& FilterClassName, const FName& CollectionName, const TSet<FSoftObjectPath>& ObjectsToRemove);
+	bool ReorderCollection(const FName& FilterClassName, const FName& CollectionToMoveName, const FName& CollectionInsertBeforeName);
 
-	void RemoveCollection(const FName& FilterClassName, const FName& CollectionName);
+	bool RenameCollection(const FName& FilterClassName, const FName CollectionNameToRename, const FName NewCollectionName);
 
-	void ReorderCollection(const FName& FilterClassName, const FName& CollectionToMoveName, const FName& CollectionInsertBeforeName);
+	bool DoesCollectionExist(const FName& FilterClassName, const FName& CollectionName);
 
 	bool IsObjectInCollection(const FName& FilterClassName, const FName& CollectionName, const FSoftObjectPath& InObject);
 
@@ -129,4 +138,9 @@ public:
 	bool IsColumnDataSerialized(const FName& FilterClassName, const FName& ColumnName);
 	
 	bool ShouldShowColumn(const FName& FilterClassName, const FName& ColumnName);
+
+	static FName AllCollectionName;
+
+	FOnObjectMixerCollectionChanged OnObjectMixerCollectionMapChanged;
+	FOnObjectMixerCollectionRenamed OnObjectMixerCollectionRenamed;
 };
diff --git a/Engine/Plugins/Editor/ObjectMixer/ObjectMixer/Source/ObjectMixer/Public/Views/MainPanel/ObjectMixerEditorMainPanel.h b/Engine/Plugins/Editor/ObjectMixer/ObjectMixer/Source/ObjectMixer/Public/Views/MainPanel/ObjectMixerEditorMainPanel.h
--- a/Engine/Plugins/Editor/ObjectMixer/ObjectMixer/Source/ObjectMixer/Public/Views/MainPanel/ObjectMixerEditorMainPanel.h
+++ b/Engine/Plugins/Editor/ObjectMixer/ObjectMixer/Source/ObjectMixer/Public/Views/MainPanel/ObjectMixerEditorMainPanel.h
@@ -3,16 +3,17 @@
 #pragma once
 
 #include "ObjectFilter/ObjectMixerEditorObjectFilter.h"
-#include "Views/List/ObjectMixerEditorListFilters/IObjectMixerEditorListFilter.h"
 
 #include "UObject/StrongObjectPtr.h"
 #include "Widgets/SWidget.h"
 
+class UObjectMixerEditorSerializedData;
+struct FObjectMixerEditorListRow;
+class IObjectMixerEditorListFilter;
+class UObjectMixerObjectFilter;
 class FObjectMixerEditorList;
 class SObjectMixerEditorMainPanel;
 
-DECLARE_MULTICAST_DELEGATE(FOnObjectMixerCollectionMapChanged)
-
 class OBJECTMIXEREDITOR_API FObjectMixerEditorMainPanel : public TSharedFromThis<FObjectMixerEditorMainPanel>
 {
 public:
@@ -160,19 +161,19 @@ public:
 	// User Collections
 
 	/**
-	 * Add set of objects to a collection in the map, or create a new collection if one does not exist,
+	 * Get a pointer to the UObjectMixerEditorSerializedData object along with the name of the filter represented by this MainPanel instance.
 	 */
-	void AddObjectsToCollection(const FName& CollectionName, const TSet<FSoftObjectPath>& ObjectsToAdd) const;
-	void RemoveObjectsFromCollection(const FName& CollectionName, const TSet<FSoftObjectPath>& ObjectsToRemove) const;
-	void RemoveCollection(const FName& CollectionName) const;
-	void ReorderCollection(const FName& CollectionToMoveName, const FName& CollectionInsertBeforeName) const;
+	UObjectMixerEditorSerializedData* GetSerializedDataOutputtingFilterName(FName& OutFilterName) const;
+	bool RequestAddObjectsToCollection(const FName& CollectionName, const TSet<FSoftObjectPath>& ObjectsToAdd) const;
+	bool RequestRemoveObjectsFromCollection(const FName& CollectionName, const TSet<FSoftObjectPath>& ObjectsToRemove) const;
+	bool RequestRemoveCollection(const FName& CollectionName) const;
+	bool RequestDuplicateCollection(const FName& CollectionToDuplicateName, FName& DesiredDuplicateName) const;
+	bool RequestReorderCollection(const FName& CollectionToMoveName, const FName& CollectionInsertBeforeName) const;
+	bool RequestRenameCollection(const FName& CollectionNameToRename, const FName& NewCollectionName) const;
+	bool DoesCollectionExist(const FName& CollectionName) const;
 	bool IsObjectInCollection(const FName& CollectionName, const FSoftObjectPath& InObject) const;
 	TSet<FName> GetCollectionsForObject(const FSoftObjectPath& InObject) const;
 	TArray<FName> GetAllCollectionNames() const;
-	FOnObjectMixerCollectionMapChanged& GetOnObjectMixerCollectionMapChanged()
-	{
-		return OnObjectMixerCollectionMapChanged;
-	}
 
 	/**
 	 * Returns the collections selected by the user. If the set is empty, consider "All" collections to be selected.
@@ -200,6 +201,4 @@ private:
 	TSet<TWeakPtr<FObjectMixerEditorListRow>> SoloRows = {};
 
 	FName ModuleName = NAME_None;
-
-	FOnObjectMixerCollectionMapChanged OnObjectMixerCollectionMapChanged;
 };
diff --git a/Engine/Plugins/Editor/ObjectMixer/ObjectMixer/Source/ObjectMixer/Public/Views/MainPanel/SObjectMixerEditorMainPanel.h b/Engine/Plugins/Editor/ObjectMixer/ObjectMixer/Source/ObjectMixer/Public/Views/MainPanel/SObjectMixerEditorMainPanel.h
--- a/Engine/Plugins/Editor/ObjectMixer/ObjectMixer/Source/ObjectMixer/Public/Views/MainPanel/SObjectMixerEditorMainPanel.h
+++ b/Engine/Plugins/Editor/ObjectMixer/ObjectMixer/Source/ObjectMixer/Public/Views/MainPanel/SObjectMixerEditorMainPanel.h
@@ -3,8 +3,6 @@
 #pragma once
 
 #include "ObjectFilter/ObjectMixerEditorObjectFilter.h"
-#include "Views/List/ObjectMixerEditorListFilters/IObjectMixerEditorListFilter.h"
-
 #include "Styling/SlateTypes.h"
 #include "Widgets/SCompoundWidget.h"
 
@@ -52,7 +50,10 @@ public:
 	
 	void RebuildCollectionSelector();
 
-	void RequestRemoveCollection(const FName& CollectionName);
+	bool RequestRemoveCollection(const FName& CollectionName);
+	bool RequestDuplicateCollection(const FName& CollectionToDuplicateName, FName& DesiredDuplicateName) const;
+	bool RequestRenameCollection(const FName& CollectionNameToRename, const FName& NewCollectionName);
+	bool DoesCollectionExist(const FName& CollectionName) const;
 	
 	void OnCollectionCheckedStateChanged(ECheckBoxState State, FName CollectionName);
 	ECheckBoxState IsCollectionChecked(FName Section) const;
diff --git a/Engine/Plugins/Editor/ObjectMixer/ObjectMixer/Source/ObjectMixer/Public/Views/Widgets/ObjectMixerEditorListMenuContext.h b/Engine/Plugins/Editor/ObjectMixer/ObjectMixer/Source/ObjectMixer/Public/Views/Widgets/ObjectMixerEditorListMenuContext.h
--- a/Engine/Plugins/Editor/ObjectMixer/ObjectMixer/Source/ObjectMixer/Public/Views/Widgets/ObjectMixerEditorListMenuContext.h
+++ b/Engine/Plugins/Editor/ObjectMixer/ObjectMixer/Source/ObjectMixer/Public/Views/Widgets/ObjectMixerEditorListMenuContext.h
@@ -33,7 +33,6 @@ public:
 private:
 
 	static void OnTextCommitted(const FText& InText, ETextCommit::Type InCommitType, const FObjectMixerEditorListMenuContextData ContextData);
-	static void OnTextChanged(const FText& InText, const FObjectMixerEditorListMenuContextData ContextData);
 	
 	static void OnClickCollectionMenuEntry(const FName Key, const FObjectMixerEditorListMenuContextData ContextData);
 	static void AddObjectsToCollection(const FName Key, const FObjectMixerEditorListMenuContextData ContextData);
diff --git a/Engine/Plugins/Editor/ObjectMixer/ObjectMixer/Source/ObjectMixer/Public/Views/Widgets/SCollectionSelectionButton.h b/Engine/Plugins/Editor/ObjectMixer/ObjectMixer/Source/ObjectMixer/Public/Views/Widgets/SCollectionSelectionButton.h
--- a/Engine/Plugins/Editor/ObjectMixer/ObjectMixer/Source/ObjectMixer/Public/Views/Widgets/SCollectionSelectionButton.h
+++ b/Engine/Plugins/Editor/ObjectMixer/ObjectMixer/Source/ObjectMixer/Public/Views/Widgets/SCollectionSelectionButton.h
@@ -8,6 +8,9 @@
 
 #define LOCTEXT_NAMESPACE "ObjectMixerEditor"
 
+class SInlineEditableTextBlock;
+class SMenuAnchor;
+
 class SCollectionSelectionButton final : public SCompoundWidget
 {
 public:
@@ -19,6 +22,10 @@ public:
 	
 	void Construct(const FArguments& InArgs, const TSharedRef<class SObjectMixerEditorMainPanel> MainPanelWidget, const FName& InCollectionName);
 
+	TSharedRef<SWidget> GetContextMenu() const;
+
+	void OnEditableTextCommitted(const FText& Text, ETextCommit::Type Type);
+
 	virtual FReply OnMouseButtonDown(const FGeometry& MyGeometry, const FPointerEvent& MouseEvent) override;
 
 	virtual FReply OnMouseButtonUp(const FGeometry& MyGeometry, const FPointerEvent& MouseEvent) override;
@@ -50,6 +57,9 @@ private:
 	FSlateRoundedBoxBrush CheckedHoveredImage = FSlateRoundedBoxBrush(FStyleColors::PrimaryHover, 4.0f, FStyleColors::Input, 1.0f);
 	FSlateRoundedBoxBrush CheckedPressedImage = FSlateRoundedBoxBrush(FStyleColors::PrimaryHover, 4.0f, FStyleColors::Input, 1.0f);
 	FSlateRoundedBoxBrush CheckedValidDropImage = FSlateRoundedBoxBrush(FStyleColors::Secondary, 4.0f, FStyleColors::Input, 1.0f);
+	
+	TSharedPtr<SInlineEditableTextBlock> EditableTextBlock;
+	TSharedPtr<SMenuAnchor> MenuAnchor;
 };
 
 #undef LOCTEXT_NAMESPACE
