Object Mixer:

-Lights on sub-levels not listed

diff --git a/Engine/Plugins/Editor/ObjectMixer/ObjectMixer/Source/ObjectMixer/Private/Views/List/SObjectMixerEditorList.cpp b/Engine/Plugins/Editor/ObjectMixer/ObjectMixer/Source/ObjectMixer/Private/Views/List/SObjectMixerEditorList.cpp
--- a/Engine/Plugins/Editor/ObjectMixer/ObjectMixer/Source/ObjectMixer/Private/Views/List/SObjectMixerEditorList.cpp
+++ b/Engine/Plugins/Editor/ObjectMixer/ObjectMixer/Source/ObjectMixer/Private/Views/List/SObjectMixerEditorList.cpp
@@ -1200,15 +1200,23 @@ void SObjectMixerEditorList::GenerateTreeView()
 	check(GEditor);
 	const UWorld* EditorWorld = GEditor->GetEditorWorldContext().World();
 	
+	
 	TSet<UObject*> AllMatchingObjects;
-	ForEachObjectWithOuter(EditorWorld,
-		[this, &AllMatchingObjects](UObject* Object)
+
+	for (const ULevel* Level : EditorWorld->GetLevels())
 	{
-		if (DoesValidWorldObjectHaveAcceptableClass(Object, ObjectClassesToFilterCache))
+		if (Level && Level->bIsVisible)
 		{
-			AllMatchingObjects.Add(Object);
+			ForEachObjectWithOuter(Level,
+			  [this, &AllMatchingObjects](UObject* Object)
+		  {
+			  if (DoesValidWorldObjectHaveAcceptableClass(Object, ObjectClassesToFilterCache))
+			  {
+				  AllMatchingObjects.Add(Object);
+			  }
+		  });
 		}
-	});
+	}
 	
 	TMap<FName, FObjectMixerEditorListRowPtr> FolderMap;
 	for (UObject* Object : AllMatchingObjects)
