Object Mixer:

-Fix: Per-filter tree item state caches

diff --git a/Engine/Plugins/Editor/ObjectMixer/ObjectMixer/Source/ObjectMixer/Private/Views/List/ObjectMixerEditorList.cpp b/Engine/Plugins/Editor/ObjectMixer/ObjectMixer/Source/ObjectMixer/Private/Views/List/ObjectMixerEditorList.cpp
--- a/Engine/Plugins/Editor/ObjectMixer/ObjectMixer/Source/ObjectMixer/Private/Views/List/ObjectMixerEditorList.cpp
+++ b/Engine/Plugins/Editor/ObjectMixer/ObjectMixer/Source/ObjectMixer/Private/Views/List/ObjectMixerEditorList.cpp
@@ -4,6 +4,14 @@
 
 #include "Views/List/SObjectMixerEditorList.h"
 
+FObjectMixerEditorList::FObjectMixerEditorList(TSharedRef<FObjectMixerEditorMainPanel, ESPMode::ThreadSafe> InMainPanel)
+{
+	MainPanelModelPtr = InMainPanel;
+
+	InMainPanel->OnPreFilterChange.AddRaw(this, &FObjectMixerEditorList::OnPreFilterChange);
+	InMainPanel->OnPostFilterChange.AddRaw(this, &FObjectMixerEditorList::OnPostFilterChange);
+}
+
 FObjectMixerEditorList::~FObjectMixerEditorList()
 {
 	FlushWidget();
@@ -11,6 +19,12 @@ FObjectMixerEditorList::~FObjectMixerEditorList()
 
 void FObjectMixerEditorList::FlushWidget()
 {
+	if (TSharedPtr<FObjectMixerEditorMainPanel> MainPanelPinned = MainPanelModelPtr.Pin())
+	{
+		MainPanelPinned->OnPreFilterChange.RemoveAll(this);
+		MainPanelPinned->OnPostFilterChange.RemoveAll(this);
+	}
+	
 	ListWidget.Reset();
 }
 
@@ -26,6 +40,29 @@ TSharedRef<SWidget> FObjectMixerEditorList::GetOrCreateWidget()
 	return ListWidget.ToSharedRef();
 }
 
+void FObjectMixerEditorList::OnPreFilterChange()
+{
+	if (ListWidget.IsValid())
+	{
+		if (const TSharedPtr<FObjectMixerEditorMainPanel> PinnedMainPanel = GetMainPanelModel().Pin())
+		{		
+			ListWidget->CacheTreeState(PinnedMainPanel->GetWeakActiveListFiltersSortedByName());
+		}
+	}
+}
+
+void FObjectMixerEditorList::OnPostFilterChange()
+{
+	if (ListWidget.IsValid())
+	{
+		if (const TSharedPtr<FObjectMixerEditorMainPanel> PinnedMainPanel = GetMainPanelModel().Pin())
+		{		
+			ListWidget->EvaluateIfRowsPassFilters();
+			ListWidget->RestoreTreeState(PinnedMainPanel->GetWeakActiveListFiltersSortedByName());
+		}
+	}
+}
+
 void FObjectMixerEditorList::ClearList() const
 {
 	if (ListWidget.IsValid())
diff --git a/Engine/Plugins/Editor/ObjectMixer/ObjectMixer/Source/ObjectMixer/Private/Views/List/ObjectMixerEditorListRow.cpp b/Engine/Plugins/Editor/ObjectMixer/ObjectMixer/Source/ObjectMixer/Private/Views/List/ObjectMixerEditorListRow.cpp
--- a/Engine/Plugins/Editor/ObjectMixer/ObjectMixer/Source/ObjectMixer/Private/Views/List/ObjectMixerEditorListRow.cpp
+++ b/Engine/Plugins/Editor/ObjectMixer/ObjectMixer/Source/ObjectMixer/Private/Views/List/ObjectMixerEditorListRow.cpp
@@ -2,8 +2,12 @@
 
 #include "Views/List/ObjectMixerEditorListRow.h"
 
+#include "ObjectMixerEditorSettings.h"
+#include "Views/List/ObjectMixerEditorListFilters/ObjectMixerEditorListFilter_Collection.h"
+
 #include "Algo/AllOf.h"
 #include "ClassIconFinder.h"
+#include "ObjectMixerEditorSerializedData.h"
 #include "GameFramework/Actor.h"
 #include "Styling/SlateIconFinder.h"
 #include "Views/List/ObjectMixerEditorList.h"
@@ -59,31 +63,25 @@ UObjectMixerObjectFilter* FObjectMixerEditorListRow::GetObjectFilter() const
 	return nullptr;
 }
 
-bool FObjectMixerEditorListRow::IsObjectRefInSelectedCollections() const
-{
-	check (ListViewPtr.IsValid());
-	
+bool FObjectMixerEditorListRow::IsObjectRefInCollection(const FName& CollectionName) const
+{	
 	if (RowType != EObjectMixerEditorListRowType::None && RowType != EObjectMixerEditorListRowType::Folder)
 	{
-		if (const TSharedPtr<FObjectMixerEditorList> ListModel = ListViewPtr.Pin()->GetListModelPtr().Pin())
+		if (CollectionName == UObjectMixerEditorSerializedData::AllCollectionName)
 		{
-			if (const TSharedPtr<FObjectMixerEditorMainPanel> MainPanel = ListModel->GetMainPanelModel().Pin())
-			{
-				const TSet<FName>& CollectionSelection = MainPanel->GetCurrentCollectionSelection();
-
-				if (CollectionSelection.Num() == 0)
-				{
-					return true;
-				}
-
-				const TSet<FName>& ObjectAssignedCollections = MainPanel->GetCollectionsForObject(GetObject());
-
-				const TSet<FName> Intersection = CollectionSelection.Intersect(ObjectAssignedCollections);
-
-				return Intersection.Num() > 0;
-			}
+			return true;
+		}
+		
+		UObjectMixerEditorSerializedData* SerializedData = GetMutableDefault<UObjectMixerEditorSerializedData>();
+		check(SerializedData);
+		
+		if (const UObjectMixerObjectFilter* Filter = GetObjectFilter())
+		{
+			return SerializedData->IsObjectInCollection(
+				Filter->GetClass()->GetFName(), CollectionName, GetObject());
 		}
 	}
+	
 	return false;
 }
 
diff --git a/Engine/Plugins/Editor/ObjectMixer/ObjectMixer/Source/ObjectMixer/Private/Views/List/SObjectMixerEditorList.cpp b/Engine/Plugins/Editor/ObjectMixer/ObjectMixer/Source/ObjectMixer/Private/Views/List/SObjectMixerEditorList.cpp
--- a/Engine/Plugins/Editor/ObjectMixer/ObjectMixer/Source/ObjectMixer/Private/Views/List/SObjectMixerEditorList.cpp
+++ b/Engine/Plugins/Editor/ObjectMixer/ObjectMixer/Source/ObjectMixer/Private/Views/List/SObjectMixerEditorList.cpp
@@ -25,6 +25,7 @@
 #include "Editor/UnrealEdEngine.h"
 #include "PlacementMode/Public/IPlacementModeModule.h"
 #include "ScopedTransaction.h"
+#include "Algo/IndexOf.h"
 #include "Styling/StyleColors.h"
 #include "Views/Widgets/ObjectMixerEditorListMenuContext.h"
 #include "Widgets/Layout/SScrollBox.h"
@@ -381,7 +382,7 @@ EObjectMixerTreeViewMode SObjectMixerEditorList::GetTreeViewMode()
 	return PinnedMainPanel->GetTreeViewMode();
 }
 
-const TArray<TSharedRef<IObjectMixerEditorListFilter>>& SObjectMixerEditorList::GetShowFilters()
+const TArray<TSharedRef<IObjectMixerEditorListFilter>>& SObjectMixerEditorList::GetListFilters()
 {
 	const TSharedPtr<FObjectMixerEditorList> PinnedListModel = GetListModelPtr().Pin();
 	check(PinnedListModel);
@@ -389,7 +390,7 @@ const TArray<TSharedRef<IObjectMixerEditorListFilter>>& SObjectMixerEditorList::
 	const TSharedPtr<FObjectMixerEditorMainPanel> PinnedMainPanel = PinnedListModel->GetMainPanelModel().Pin();
 	check(PinnedMainPanel);
 	
-	return PinnedMainPanel->GetShowFilters();
+	return PinnedMainPanel->GetListFilters();
 }
 
 void SObjectMixerEditorList::EvaluateIfRowsPassFilters(const bool bShouldRefreshAfterward)
@@ -399,7 +400,7 @@ void SObjectMixerEditorList::EvaluateIfRowsPassFilters(const bool bShouldRefresh
 	TSet<TSharedRef<IObjectMixerEditorListFilter>> MatchAnyOfFilters;
 	TSet<TSharedRef<IObjectMixerEditorListFilter>> MatchAllOfFilters;
 
-	for (const TSharedRef<IObjectMixerEditorListFilter>& Filter : GetShowFilters())
+	for (const TSharedRef<IObjectMixerEditorListFilter>& Filter : GetListFilters())
 	{
 		if (Filter->GetFilterMatchType() ==
 			IObjectMixerEditorListFilter::EObjectMixerEditorListFilterMatchType::MatchAll)
@@ -877,13 +878,13 @@ void SObjectMixerEditorList::SetAllGroupsCollapsed()
 	}
 }
 
-void SObjectMixerEditorList::CacheTreeState()
+void SObjectMixerEditorList::CacheTreeState(const TArray<TWeakPtr<IObjectMixerEditorListFilter>>& InFilterCombo)
 {
 	struct Local
 	{
 		static void RecursivelyCacheTreeState(
 			const TArray<FObjectMixerEditorListRowPtr>& InObjects,
-			TArray<FTreeItemStateCache>& TreeItemStateCache,
+			TArray<FTreeItemStateCache>* TreeItemStateCache,
 			TSharedPtr<STreeView<FObjectMixerEditorListRowPtr>> TreeViewPtr)
 		{
 			for (const TSharedPtr<FObjectMixerEditorListRow>& TreeViewItem : InObjects)
@@ -895,7 +896,7 @@ void SObjectMixerEditorList::CacheTreeState()
 
 				if (!RowName.IsEmpty())
 				{
-					TreeItemStateCache.Add(
+					TreeItemStateCache->Add(
 						{
 							RowObject ? RowObject->GetUniqueID() : -1,
 							RowName,
@@ -910,23 +911,40 @@ void SObjectMixerEditorList::CacheTreeState()
 		}
 	};
 
-	Local::RecursivelyCacheTreeState(TreeViewRootObjects, TreeItemStateCache, TreeViewPtr);
+	FFilterComboToStateCaches* Match = Algo::FindByPredicate(
+		FilterComboToStateCaches,
+		[&InFilterCombo](
+		const FFilterComboToStateCaches& Cache)
+		{
+			return Cache.FilterCombo == InFilterCombo;
+		});
+
+	if (Match)
+	{
+		Local::RecursivelyCacheTreeState(TreeViewRootObjects, &Match->Caches, TreeViewPtr);
+	}
+	else
+	{
+		TArray<FTreeItemStateCache> NewCache;
+		Local::RecursivelyCacheTreeState(TreeViewRootObjects, &NewCache, TreeViewPtr);
+		FilterComboToStateCaches.Add({InFilterCombo, MoveTemp(NewCache)});
+	}
 }
 
-void SObjectMixerEditorList::RestoreTreeState(const bool bFlushCache)
+void SObjectMixerEditorList::RestoreTreeState(const TArray<TWeakPtr<IObjectMixerEditorListFilter>>& InFilterCombo, const bool bFlushCache)
 {
 	struct Local
 	{
 		static void RecursivelyRestoreTreeState(
 			const TArray<FObjectMixerEditorListRowPtr>& InObjects,
-			TArray<FTreeItemStateCache>& TreeItemStateCache,
+			TArray<FTreeItemStateCache>* TreeItemStateCache,
 			TSharedPtr<STreeView<FObjectMixerEditorListRowPtr>> TreeViewPtr,
 			const bool bExpandByDefault)
 		{
 			for (const TSharedPtr<FObjectMixerEditorListRow>& TreeViewItem : InObjects)
 			{
 				if (const FTreeItemStateCache* StateCachePtr = Algo::FindByPredicate(
-					TreeItemStateCache,
+					*TreeItemStateCache,
 					[TreeViewItem](const FTreeItemStateCache& Other)
 					{
 						if (Other.UniqueId != -1)
@@ -964,13 +982,21 @@ void SObjectMixerEditorList::RestoreTreeState(const bool bFlushCache)
 	{
 		bExpandByDefault = Settings->bExpandTreeViewItemsByDefault;
 	}
-	
-	Local::RecursivelyRestoreTreeState(
-		TreeViewRootObjects, TreeItemStateCache, TreeViewPtr, bExpandByDefault);
 
-	if (bFlushCache)
+	for (int32 CachesItr = FilterComboToStateCaches.Num() - 1; CachesItr >= 0; CachesItr--)
 	{
-		TreeItemStateCache.Empty();
+		if (FilterComboToStateCaches[CachesItr].FilterCombo == InFilterCombo)
+		{
+			Local::RecursivelyRestoreTreeState(
+			  TreeViewRootObjects, &FilterComboToStateCaches[CachesItr].Caches, TreeViewPtr, bExpandByDefault);
+
+			if (bFlushCache)
+			{
+				FilterComboToStateCaches.RemoveAt(CachesItr);
+			}
+
+			break;
+		}
 	}
 }
 
@@ -1183,7 +1209,7 @@ void CreateFolderRowsForTopLevelRow(
 
 void SObjectMixerEditorList::GenerateTreeView()
 {
-	TSharedPtr<FObjectMixerEditorList> PinnedListModel = ListModelPtr.Pin();
+	const TSharedPtr<FObjectMixerEditorList> PinnedListModel = ListModelPtr.Pin();
 	check(PinnedListModel);
 	
 	if (!ensure(TreeViewPtr.IsValid()))
@@ -1191,7 +1217,14 @@ void SObjectMixerEditorList::GenerateTreeView()
 		return;
 	}
 
-	CacheTreeState();
+	TArray<TWeakPtr<IObjectMixerEditorListFilter>> FilterCombo = {};
+
+	if (const TSharedPtr<FObjectMixerEditorMainPanel> PinnedMainPanel = PinnedListModel->GetMainPanelModel().Pin())
+	{
+		FilterCombo = PinnedMainPanel->GetWeakActiveListFiltersSortedByName();
+	}
+
+	CacheTreeState(FilterCombo);
 	
 	FlushMemory(true);
 
@@ -1200,7 +1233,6 @@ void SObjectMixerEditorList::GenerateTreeView()
 	check(GEditor);
 	const UWorld* EditorWorld = GEditor->GetEditorWorldContext().World();
 	
-	
 	TSet<UObject*> AllMatchingObjects;
 
 	for (const ULevel* Level : EditorWorld->GetLevels())
@@ -1248,7 +1280,7 @@ void SObjectMixerEditorList::GenerateTreeView()
 
 	RefreshList();
 
-	RestoreTreeState();
+	RestoreTreeState(FilterCombo);
 }
 
 void SObjectMixerEditorList::FindVisibleTreeViewObjects()
@@ -1332,7 +1364,11 @@ void SObjectMixerEditorList::OnGetRowChildren(FObjectMixerEditorListRowPtr Row,
 
 		if (Row->GetShouldExpandAllChildren())
 		{
+			Row->SetShouldExpandAllChildren(false);
 			SetChildExpansionRecursively(Row, true);
+		}
+		else
+		{
 			Row->SetShouldExpandAllChildren(false);
 		}
 	}
@@ -1357,7 +1393,6 @@ void SObjectMixerEditorList::OnRowChildExpansionChange(FObjectMixerEditorListRow
 			}
 		}
 		
-		TreeViewPtr->SetItemExpansion(Row, bIsExpanded);
 		Row->SetIsTreeViewItemExpanded(bIsExpanded);
 	}
 }
diff --git a/Engine/Plugins/Editor/ObjectMixer/ObjectMixer/Source/ObjectMixer/Private/Views/MainPanel/ObjectMixerEditorMainPanel.cpp b/Engine/Plugins/Editor/ObjectMixer/ObjectMixer/Source/ObjectMixer/Private/Views/MainPanel/ObjectMixerEditorMainPanel.cpp
--- a/Engine/Plugins/Editor/ObjectMixer/ObjectMixer/Source/ObjectMixer/Private/Views/MainPanel/ObjectMixerEditorMainPanel.cpp
+++ b/Engine/Plugins/Editor/ObjectMixer/ObjectMixer/Source/ObjectMixer/Private/Views/MainPanel/ObjectMixerEditorMainPanel.cpp
@@ -105,9 +105,14 @@ void FObjectMixerEditorMainPanel::CacheObjectFilterObject()
 	}
 }
 
-const TArray<TSharedRef<IObjectMixerEditorListFilter>>& FObjectMixerEditorMainPanel::GetShowFilters() const
+const TArray<TSharedRef<IObjectMixerEditorListFilter>>& FObjectMixerEditorMainPanel::GetListFilters() const
 {
-	return MainPanelWidget->GetShowFilters();	
+	return MainPanelWidget->GetListFilters();	
+}
+
+TArray<TWeakPtr<IObjectMixerEditorListFilter>> FObjectMixerEditorMainPanel::GetWeakActiveListFiltersSortedByName() const
+{
+	return MainPanelWidget->GetWeakActiveListFiltersSortedByName();
 }
 
 UObjectMixerEditorSerializedData* FObjectMixerEditorMainPanel::GetSerializedDataOutputtingFilterName(FName& OutFilterName) const
@@ -232,7 +237,7 @@ TArray<FName> FObjectMixerEditorMainPanel::GetAllCollectionNames() const
 	return {};
 }
 
-const TSet<FName>& FObjectMixerEditorMainPanel::GetCurrentCollectionSelection() const
+TSet<TSharedRef<FObjectMixerEditorListFilter_Collection>> FObjectMixerEditorMainPanel::GetCurrentCollectionSelection() const
 {
 	return MainPanelWidget->GetCurrentCollectionSelection();
 }
diff --git a/Engine/Plugins/Editor/ObjectMixer/ObjectMixer/Source/ObjectMixer/Private/Views/MainPanel/SObjectMixerEditorMainPanel.cpp b/Engine/Plugins/Editor/ObjectMixer/ObjectMixer/Source/ObjectMixer/Private/Views/MainPanel/SObjectMixerEditorMainPanel.cpp
--- a/Engine/Plugins/Editor/ObjectMixer/ObjectMixer/Source/ObjectMixer/Private/Views/MainPanel/SObjectMixerEditorMainPanel.cpp
+++ b/Engine/Plugins/Editor/ObjectMixer/ObjectMixer/Source/ObjectMixer/Private/Views/MainPanel/SObjectMixerEditorMainPanel.cpp
@@ -4,7 +4,6 @@
 
 #include "ObjectMixerEditorLog.h"
 #include "ObjectMixerEditorModule.h"
-#include "ObjectMixerEditorSerializedData.h"
 #include "ObjectMixerEditorSettings.h"
 #include "Views/List/ObjectMixerEditorList.h"
 #include "Views/List/ObjectMixerEditorListFilters/ObjectMixerEditorListFilter_Collection.h"
@@ -18,6 +17,7 @@
 #include "Kismet2/SClassPickerDialog.h"
 #include "PlacementMode/Public/IPlacementModeModule.h"
 #include "SPositiveActionButton.h"
+#include "Algo/RemoveIf.h"
 #include "Styling/StyleColors.h"
 #include "Widgets/Input/SCheckBox.h"
 #include "Widgets/Input/SComboButton.h"
@@ -69,7 +69,15 @@ void SObjectMixerEditorMainPanel::Construct(
 		SerializedData->OnObjectMixerCollectionMapChanged.AddRaw(this, &SObjectMixerEditorMainPanel::RebuildCollectionSelector);
 	}
 
-	ShowFilters.Add(MakeShared<FObjectMixerEditorListFilter_Collection>());
+	// Make "All" filter which must always be present
+	const TSharedRef<FObjectMixerEditorListFilter_Collection> NewCollectionFilter =
+		MakeShared<FObjectMixerEditorListFilter_Collection>(UObjectMixerEditorSerializedData::AllCollectionName);
+		
+	ListFilters.Add(NewCollectionFilter);
+	
+	RebuildCollectionSelector();
+
+	SetSingleCollectionSelection();
 }
 
 TSharedRef<SWidget> SObjectMixerEditorMainPanel::GenerateToolbar()
@@ -406,7 +414,7 @@ TSharedRef<SWidget> SObjectMixerEditorMainPanel::BuildShowOptionsMenu()
 	}
 	ShowOptionsMenuBuilder.EndSection();
 
-	if (ShowFilters.Num())
+	if (ListFilters.Num())
 	{
 		ShowOptionsMenuBuilder.BeginSection("", LOCTEXT("ShowOptions_ShowSectionHeading", "Show"));
 		{
@@ -434,10 +442,10 @@ TSharedRef<SWidget> SObjectMixerEditorMainPanel::BuildShowOptionsMenu()
 			   );
 			};
 
-			for (const TSharedRef<IObjectMixerEditorListFilter>& Filter : ShowFilters)
+			for (const TSharedRef<IObjectMixerEditorListFilter>& Filter : ListFilters)
 			{
 				// Don't add non-toggle filters to Show Options, they're always on
-				if (Filter->IsToggleable())
+				if (Filter->IsUserToggleable())
 				{
 					AddFiltersLambda(Filter);
 				}
@@ -454,6 +462,34 @@ void SObjectMixerEditorMainPanel::OnSearchTextChanged(const FText& Text)
 	ExecuteListViewSearchOnAllRows(Text.ToString(), true);
 }
 
+void SObjectMixerEditorMainPanel::ResetCollectionFilters()
+{
+	// Clean up all collection filters
+	for (int32 FiltersItr = ListFilters.Num() - 1; FiltersItr >= 0; FiltersItr--)
+	{
+		if (const TSharedPtr<FObjectMixerEditorListFilter_Collection> CollectionFilter =
+			StaticCastSharedRef<FObjectMixerEditorListFilter_Collection>(ListFilters[FiltersItr]))
+		{
+			ListFilters.RemoveAt(FiltersItr);
+		}
+	}
+}
+
+void SObjectMixerEditorMainPanel::SetSingleCollectionSelection(const FName& CollectionToEnableName)
+{
+	// Disable all collection filters except CollectionToEnableName
+	for (const TSharedRef<IObjectMixerEditorListFilter>& Filter : GetListFilters())
+	{
+		if (const TSharedPtr<FObjectMixerEditorListFilter_Collection> CollectionFilter =
+			StaticCastSharedRef<FObjectMixerEditorListFilter_Collection>(Filter))
+		{
+			const bool bIsAllCollectionFilter =
+				CollectionFilter->CollectionName.IsEqual(CollectionToEnableName);
+			CollectionFilter->SetFilterActive(bIsAllCollectionFilter);
+		}
+	}
+}
+
 FText SObjectMixerEditorMainPanel::GetSearchTextFromSearchInputField() const
 {
 	return ensureAlwaysMsgf(SearchBoxPtr.IsValid(),
@@ -502,7 +538,7 @@ void SObjectMixerEditorMainPanel::SetTreeViewMode(EObjectMixerTreeViewMode InVie
 void SObjectMixerEditorMainPanel::ToggleFilterActive(const FString& FilterName)
 {
 	if (const TSharedRef<IObjectMixerEditorListFilter>* Match =
-		Algo::FindByPredicate(GetShowFilters(),
+		Algo::FindByPredicate(GetListFilters(),
 		[&FilterName](TSharedRef<IObjectMixerEditorListFilter> Comparator)
 		{
 			return Comparator->GetFilterName().Equals(FilterName);
@@ -515,43 +551,144 @@ void SObjectMixerEditorMainPanel::ToggleFilterActive(const FString& FilterName)
 	}
 }
 
-const TSet<FName>& SObjectMixerEditorMainPanel::GetCurrentCollectionSelection()
+TArray<TWeakPtr<IObjectMixerEditorListFilter>> SObjectMixerEditorMainPanel::GetWeakActiveListFiltersSortedByName()
 {
-	return CurrentCollectionSelection;
+	TArray<TWeakPtr<IObjectMixerEditorListFilter>> ActiveFilters = {};
+	
+	for (const TSharedRef<IObjectMixerEditorListFilter>& ListFilter : GetListFilters())
+	{
+		if (ListFilter->GetIsFilterActive())
+		{
+			ActiveFilters.Add(ListFilter);
+		}
+	}
+
+	ActiveFilters.StableSort(
+		[](const TWeakPtr<IObjectMixerEditorListFilter>& A, const TWeakPtr<IObjectMixerEditorListFilter>& B)
+		{
+			FString NameA = A.IsValid() ? A.Pin()->GetFilterName() : "";
+			FString NameB = A.IsValid() ? A.Pin()->GetFilterName() : "";
+			return NameA < NameB;
+		});
+
+	return ActiveFilters;
+}
+
+TSet<TSharedRef<FObjectMixerEditorListFilter_Collection>> SObjectMixerEditorMainPanel::GetCurrentCollectionSelection()
+{
+	TSet<TSharedRef<FObjectMixerEditorListFilter_Collection>> CollectionFilters;
+	for (const TSharedRef<IObjectMixerEditorListFilter>& ListFilter : GetListFilters())
+	{
+		if (const TSharedPtr<FObjectMixerEditorListFilter_Collection> CollectionFilter =
+			StaticCastSharedRef<FObjectMixerEditorListFilter_Collection>(ListFilter))
+		{
+			CollectionFilters.Add(CollectionFilter.ToSharedRef());
+		}
+	}
+	
+	return CollectionFilters;
 }
 
 void SObjectMixerEditorMainPanel::RebuildCollectionSelector()
 {
 	check(MainPanelModel.IsValid());
-	
+
+	// Make user collections
+
 	CollectionSelectorBox->ClearChildren();
 	CollectionSelectorBox->SetVisibility(EVisibility::Collapsed);
 
-	TArray<FName> AllCollections = MainPanelModel.Pin()->GetAllCollectionNames();
+	auto CreateCollectionFilterAndAddToCollectionSelector =
+		[this](const TSharedRef<FObjectMixerEditorListFilter_Collection> NewCollectionFilter) 
+	{
+		CollectionSelectorBox->AddSlot()
+		[
+			SNew(SCollectionSelectionButton, SharedThis(this), NewCollectionFilter)
+		];
+	};
+
+	TArray<FName> AllCollectionNames = MainPanelModel.Pin()->GetAllCollectionNames();
 	
-	if (AllCollections.IsEmpty())
+	// Remove any collection filters that no longer match (except "All")
+	ListFilters.SetNum(Algo::StableRemoveIf(ListFilters,
+		[&AllCollectionNames](const TSharedRef<IObjectMixerEditorListFilter>& ListFilter)
+		{
+			if (const TSharedPtr<FObjectMixerEditorListFilter_Collection> CollectionFilter =
+				StaticCastSharedRef<FObjectMixerEditorListFilter_Collection>(ListFilter))
+			{
+				const bool bIsAllCollection = CollectionFilter->CollectionName.IsEqual(UObjectMixerEditorSerializedData::AllCollectionName);
+				const bool bIsFilterNameFoundInAllCollectionNames = AllCollectionNames.Contains(CollectionFilter->CollectionName);
+				
+				return !bIsAllCollection && !bIsFilterNameFoundInAllCollectionNames;
+			}
+			return false;
+		})
+	);
+
+	const TSharedRef<IObjectMixerEditorListFilter>* AllCollection =
+		Algo::FindByPredicate(
+			ListFilters,
+			[&AllCollectionNames](const TSharedRef<IObjectMixerEditorListFilter>& ListFilter)
+			{
+				if (const TSharedPtr<FObjectMixerEditorListFilter_Collection> CollectionFilter =
+					StaticCastSharedRef<FObjectMixerEditorListFilter_Collection>(ListFilter))
+				{
+					return CollectionFilter->CollectionName.IsEqual(UObjectMixerEditorSerializedData::AllCollectionName);
+				}
+				return false;
+			});
+	
+	// No collections - rather than show "All", just keep the box hidden
+	if (AllCollectionNames.IsEmpty())
 	{
-		// no collections - rather than show just "All", hide the box
-		ResetCurrentCollectionSelection();
+		// Set "All" filter to active
+		if (AllCollection)
+		{
+			(*AllCollection)->SetFilterActive(true);
+		}
+		
 		return;
 	}
 
-	auto CreateSection = [this](FName CollectionName) 
+	// Make "All" collection widget
+	if (AllCollection)
 	{
-		return SNew(SCollectionSelectionButton, SharedThis(this), CollectionName);
-	};
-
-	CollectionSelectorBox->AddSlot()
-	[
-		CreateSection(UObjectMixerEditorSerializedData::AllCollectionName)
-	];
-
-	for (const FName& Key : AllCollections)
+		CreateCollectionFilterAndAddToCollectionSelector(
+			StaticCastSharedRef<FObjectMixerEditorListFilter_Collection>(*AllCollection));
+	}
+	
+	TMap<FName, TSharedRef<FObjectMixerEditorListFilter_Collection>> CollectionNamesToFilters; 
+	for (const FName& Key : AllCollectionNames)
 	{
-		CollectionSelectorBox->AddSlot()
-		[
-			CreateSection(Key)
-		];
+		// Try to find a matching filter
+		if (const TSharedRef<IObjectMixerEditorListFilter>* Match = Algo::FindByPredicate(
+			ListFilters,
+			[&Key](const TSharedRef<IObjectMixerEditorListFilter>& ListFilter)
+			{
+				if (const TSharedPtr<FObjectMixerEditorListFilter_Collection> CollectionFilter =
+					StaticCastSharedRef<FObjectMixerEditorListFilter_Collection>(ListFilter))
+				{
+					return CollectionFilter->CollectionName.IsEqual(Key);
+				}
+				return false;
+			}))
+		{
+			CollectionNamesToFilters.Add(Key, StaticCastSharedRef<FObjectMixerEditorListFilter_Collection>(*Match));
+		}
+		else
+		{
+			// Otherwise create new filter for unmatched collection name
+			TSharedRef<FObjectMixerEditorListFilter_Collection> NewCollectionFilter =
+				MakeShared<FObjectMixerEditorListFilter_Collection>(Key);
+			ListFilters.Add(NewCollectionFilter);
+			CollectionNamesToFilters.Add(Key, NewCollectionFilter);
+		}
+
+		// Then create widgets for each key
+		if (const TSharedRef<FObjectMixerEditorListFilter_Collection>* FoundCollectionFilter = CollectionNamesToFilters.Find(Key))
+		{
+			CreateCollectionFilterAndAddToCollectionSelector(*FoundCollectionFilter);
+		}
 	}
 
 	CollectionSelectorBox->SetVisibility(EVisibility::Visible);
@@ -561,7 +698,18 @@ bool SObjectMixerEditorMainPanel::RequestRemoveCollection(const FName& Collectio
 {
 	if (MainPanelModel.Pin()->RequestRemoveCollection(CollectionName))
 	{
-		CurrentCollectionSelection.Remove(CollectionName);
+		const int32 NewListFilterCount = Algo::StableRemoveIf(ListFilters,
+			[CollectionName](const TSharedRef<IObjectMixerEditorListFilter>& ListFilter)
+			{
+				if (const TSharedPtr<FObjectMixerEditorListFilter_Collection> CollectionFilter =
+					StaticCastSharedRef<FObjectMixerEditorListFilter_Collection>(ListFilter))
+				{
+					return CollectionFilter->CollectionName.IsEqual(CollectionName);
+				}
+				return false;
+			});
+		
+		ListFilters.SetNum(NewListFilterCount);
 
 		return true;
 	}
@@ -580,10 +728,30 @@ bool SObjectMixerEditorMainPanel::RequestRenameCollection(
 {
 	if (MainPanelModel.Pin()->RequestRenameCollection(CollectionNameToRename, NewCollectionName))
 	{
-		if (CurrentCollectionSelection.Contains(CollectionNameToRename))
+		if (TSharedRef<IObjectMixerEditorListFilter>* Match = Algo::FindByPredicate(ListFilters,
+			[CollectionNameToRename](const TSharedRef<IObjectMixerEditorListFilter>& ListFilter)
+			{
+				if (const TSharedPtr<FObjectMixerEditorListFilter_Collection> CollectionFilter =
+					StaticCastSharedRef<FObjectMixerEditorListFilter_Collection>(ListFilter))
+				{
+					return CollectionFilter->CollectionName.IsEqual(CollectionNameToRename);
+				}
+				return false;
+			}))
 		{
-			CurrentCollectionSelection.Remove(CollectionNameToRename);
-			CurrentCollectionSelection.Add(NewCollectionName);
+			StaticCastSharedRef<FObjectMixerEditorListFilter_Collection>(*Match)->CollectionName = NewCollectionName;
+		}
+		else
+		{
+			TSharedRef<FObjectMixerEditorListFilter_Collection> NewCollectionFilter =
+				MakeShared<FObjectMixerEditorListFilter_Collection>(NewCollectionName);
+		
+			ListFilters.Add(NewCollectionFilter);
+		
+			CollectionSelectorBox->AddSlot()
+			[
+				SNew(SCollectionSelectionButton, SharedThis(this), NewCollectionFilter)
+			];
 		}
 
 		return true;
@@ -597,60 +765,115 @@ bool SObjectMixerEditorMainPanel::DoesCollectionExist(const FName& CollectionNam
 	return MainPanelModel.Pin()->DoesCollectionExist(CollectionName);
 }
 
-void SObjectMixerEditorMainPanel::OnCollectionCheckedStateChanged(ECheckBoxState State, FName CollectionName)
+void SObjectMixerEditorMainPanel::OnCollectionCheckedStateChanged(bool bShouldBeChecked, FName CollectionName)
 {
 	check(MainPanelModel.IsValid());
 
-	// Remove Collection
-	if (FSlateApplication::Get().GetModifierKeys().IsAltDown())
-	{
-		RequestRemoveCollection(CollectionName);
-	}
-	else
-	{
-		const bool bIsControlDown = FSlateApplication::Get().GetModifierKeys().IsControlDown();
+	GetMainPanelModel().Pin()->OnPreFilterChange.Broadcast();
 
-		if (State == ECheckBoxState::Unchecked)
+	const bool bIsAllCollection = CollectionName.IsEqual(UObjectMixerEditorSerializedData::AllCollectionName);
+
+	const bool bIsControlDown = FSlateApplication::Get().GetModifierKeys().IsControlDown();
+
+	if (bShouldBeChecked)
+	{
+		if (bIsControlDown)
 		{
-			if (bIsControlDown)
+			if (bIsAllCollection)
 			{
-				CurrentCollectionSelection.Remove(CollectionName);
+				// Can't multi-select All
+				return;
 			}
-			else
+			
+			// Enable this collection and disable "All"
+			for (const TSharedRef<IObjectMixerEditorListFilter>& Filter : GetListFilters())
 			{
-				CurrentCollectionSelection.Reset();
-
-				if (CollectionName != UObjectMixerEditorSerializedData::AllCollectionName && CollectionName != NAME_None)
+				if (const TSharedPtr<FObjectMixerEditorListFilter_Collection> CollectionFilter =
+					StaticCastSharedRef<FObjectMixerEditorListFilter_Collection>(Filter))
 				{
-					CurrentCollectionSelection.Add(CollectionName);
+					if (CollectionFilter->CollectionName.IsEqual(CollectionName))
+					{
+						CollectionFilter->SetFilterActive(true);
+					}
+					else if (CollectionFilter->CollectionName.IsEqual(UObjectMixerEditorSerializedData::AllCollectionName))
+					{
+						CollectionFilter->SetFilterActive(false);
+					}
 				}
 			}
 		}
-		else if (State == ECheckBoxState::Checked)
+		else
+		{
+			SetSingleCollectionSelection(CollectionName);
+		}
+	}
+	else
+	{		
+		if (bIsControlDown)
 		{
-			if (!bIsControlDown)
+			if (bIsAllCollection)
 			{
-				CurrentCollectionSelection.Reset();
+				// Can't disable All
+				return;
+			}
+			
+			// Disable just this collection
+			int32 ActiveFilterCount = 0;
+			for (const TSharedRef<IObjectMixerEditorListFilter>& Filter : GetListFilters())
+			{
+				if (const TSharedPtr<FObjectMixerEditorListFilter_Collection> CollectionFilter =
+					StaticCastSharedRef<FObjectMixerEditorListFilter_Collection>(Filter))
+				{
+					if (CollectionFilter->CollectionName.IsEqual(CollectionName))
+					{
+						CollectionFilter->SetFilterActive(false);
+					}
+					else if (CollectionFilter->GetIsFilterActive())
+					{
+						ActiveFilterCount++;
+					}
+				}
 			}
 
-			if (CollectionName != UObjectMixerEditorSerializedData::AllCollectionName && CollectionName != NAME_None)
+			if (ActiveFilterCount == 0)
+			{
+				// Reset to all
+				SetSingleCollectionSelection();
+			}
+		}
+		else
+		{
+			if (bIsAllCollection)
 			{
-				CurrentCollectionSelection.Add(CollectionName);
+				// Reset to all
+				SetSingleCollectionSelection();
 			}
+			
+			// Set just this filter active
+			SetSingleCollectionSelection(CollectionName);
 		}
 	}
 
-	MainPanelModel.Pin()->GetEditorListModel().Pin()->EvaluateIfRowsPassFilters();
+	GetMainPanelModel().Pin()->OnPostFilterChange.Broadcast();
 }
 
-ECheckBoxState SObjectMixerEditorMainPanel::IsCollectionChecked(FName Section) const
+ECheckBoxState SObjectMixerEditorMainPanel::IsCollectionChecked(FName CollectionName) const
 {
-	if (CurrentCollectionSelection.IsEmpty() && Section == UObjectMixerEditorSerializedData::AllCollectionName)
-	{
-		return ECheckBoxState::Checked;
-	}
+	const TSharedRef<IObjectMixerEditorListFilter>* Match = Algo::FindByPredicate(ListFilters,
+		[CollectionName](const TSharedRef<IObjectMixerEditorListFilter>& ListFilter)
+		{
+			if (ListFilter->GetIsFilterActive())
+			{
+				if (const TSharedPtr<FObjectMixerEditorListFilter_Collection> CollectionFilter =
+				   StaticCastSharedRef<FObjectMixerEditorListFilter_Collection>(ListFilter))
+				{
+					return CollectionFilter->CollectionName.IsEqual(CollectionName);
+				}
+			}
+			return false;
+		});
 
-	return CurrentCollectionSelection.Contains(Section) ? ECheckBoxState::Checked : ECheckBoxState::Unchecked;
+	return Match ? ECheckBoxState::Checked : ECheckBoxState::Unchecked;
 }
 
 SObjectMixerEditorMainPanel::~SObjectMixerEditorMainPanel()
diff --git a/Engine/Plugins/Editor/ObjectMixer/ObjectMixer/Source/ObjectMixer/Private/Views/Widgets/SCollectionSelectionButton.cpp b/Engine/Plugins/Editor/ObjectMixer/ObjectMixer/Source/ObjectMixer/Private/Views/Widgets/SCollectionSelectionButton.cpp
--- a/Engine/Plugins/Editor/ObjectMixer/ObjectMixer/Source/ObjectMixer/Private/Views/Widgets/SCollectionSelectionButton.cpp
+++ b/Engine/Plugins/Editor/ObjectMixer/ObjectMixer/Source/ObjectMixer/Private/Views/Widgets/SCollectionSelectionButton.cpp
@@ -3,6 +3,7 @@
 #include "Views/Widgets/SCollectionSelectionButton.h"
 
 #include "ObjectMixerEditorSerializedData.h"
+#include "Views/List/ObjectMixerEditorListFilters/ObjectMixerEditorListFilter_Collection.h"
 #include "Views/MainPanel/ObjectMixerEditorMainPanel.h"
 #include "Views/MainPanel/SObjectMixerEditorMainPanel.h"
 
@@ -45,10 +46,13 @@ public:
 	}
 };
 
-void SCollectionSelectionButton::Construct(const FArguments& InArgs, const TSharedRef<SObjectMixerEditorMainPanel> MainPanelWidget, const FName& InCollectionName)
+void SCollectionSelectionButton::Construct(
+	const FArguments& InArgs,
+	const TSharedRef<SObjectMixerEditorMainPanel> MainPanelWidget,
+	const TSharedRef<FObjectMixerEditorListFilter_Collection> InCollectionListFilter)
 {
 	MainPanelPtr = MainPanelWidget;
-	CollectionName = InCollectionName;
+	CollectionListFilter = InCollectionListFilter;
 	
 	ChildSlot
 	[
@@ -62,7 +66,7 @@ void SCollectionSelectionButton::Construct(const FArguments& InArgs, const TShar
 			[
 				SAssignNew(EditableTextBlock, SInlineEditableTextBlock)
 				.Style(&FAppStyle::Get().GetWidgetStyle<FInlineEditableTextBlockStyle>("InlineEditableTextBlockSmallStyle"))
-				.Text(FText::FromName(CollectionName))
+				.Text_Lambda([this, PinnedFilter = CollectionListFilter.Pin()]() { return FText::FromName(PinnedFilter->CollectionName);})
 				.Visibility(EVisibility::HitTestInvisible)
 				.ColorAndOpacity(FSlateColor::UseForeground())
 				.OnTextCommitted(this, &SCollectionSelectionButton::OnEditableTextCommitted)
@@ -75,13 +79,15 @@ TSharedRef<SWidget> SCollectionSelectionButton::GetContextMenu() const
 {
 	FMenuBuilder MenuBuilder(true, nullptr);
 
+	const FName CollectionName = CollectionListFilter.Pin()->CollectionName;
+
 	MenuBuilder.AddMenuEntry(
 		FText::Format(LOCTEXT("DeleteCollectionButtonMenuEntry", "Delete '{0}'"), FText::FromName(CollectionName)),
 		LOCTEXT("DeleteCollectionButtonMenuEntryTooltip", "Remove this collection. This operation can be undone."),
 		FGenericCommands::Get().Delete->GetIcon(),
 		FUIAction(FExecuteAction::CreateLambda([this]()
 		{
-			MainPanelPtr.Pin()->RequestRemoveCollection(CollectionName);
+			MainPanelPtr.Pin()->RequestRemoveCollection(CollectionListFilter.Pin()->CollectionName);
 		}))
 	);
 
@@ -91,8 +97,9 @@ TSharedRef<SWidget> SCollectionSelectionButton::GetContextMenu() const
 		FGenericCommands::Get().Duplicate->GetIcon(),
 		FUIAction(FExecuteAction::CreateLambda([this]()
 		{
-			FName DuplicateName = CollectionName;
-			MainPanelPtr.Pin()->RequestDuplicateCollection(CollectionName, DuplicateName);
+			const TSharedPtr<FObjectMixerEditorListFilter_Collection> CollectionList = CollectionListFilter.Pin();
+			FName DuplicateName = CollectionList->CollectionName;
+			MainPanelPtr.Pin()->RequestDuplicateCollection(CollectionList->CollectionName, DuplicateName);
 		}))
 	);
 
@@ -114,10 +121,8 @@ void SCollectionSelectionButton::OnEditableTextCommitted(const FText& Text, ETex
 	if (Type == ETextCommit::OnEnter)
 	{
 		const FName NewCollectionName = *Text.ToString();
-		if (MainPanelPtr.Pin()->RequestRenameCollection(CollectionName, NewCollectionName))
-		{
-			CollectionName = NewCollectionName;
-		}
+		const TSharedPtr<FObjectMixerEditorListFilter_Collection> CollectionFilter = CollectionListFilter.Pin();
+		MainPanelPtr.Pin()->RequestRenameCollection(CollectionFilter->CollectionName, NewCollectionName);
 	}
 }
 
@@ -134,23 +139,10 @@ FReply SCollectionSelectionButton::OnMouseButtonDown(const FGeometry& MyGeometry
 FReply SCollectionSelectionButton::OnMouseButtonUp(const FGeometry& MyGeometry, const FPointerEvent& MouseEvent) 
 {
 	bIsPressed = false;
+	const FName CollectionName = CollectionListFilter.Pin()->CollectionName;
 	if (MouseEvent.GetEffectingButton() == EKeys::LeftMouseButton )
 	{
-		ECheckBoxState CurrentState = MainPanelPtr.Pin()->IsCollectionChecked(CollectionName);
-
-		ECheckBoxState NewState = ECheckBoxState::Checked;
-
-		switch (CurrentState)
-		{
-		case ECheckBoxState::Checked:
-			NewState = ECheckBoxState::Unchecked;
-			break;
-
-		default:
-			break;
-		}
-		
-		MainPanelPtr.Pin()->OnCollectionCheckedStateChanged(NewState, CollectionName);
+		MainPanelPtr.Pin()->OnCollectionCheckedStateChanged(!GetIsChecked(), CollectionName);
 	}
 	if (MouseEvent.GetEffectingButton() == EKeys::RightMouseButton)
 	{
@@ -164,6 +156,7 @@ FReply SCollectionSelectionButton::OnMouseButtonUp(const FGeometry& MyGeometry,
 
 FReply SCollectionSelectionButton::OnDragDetected(const FGeometry& MyGeometry, const FPointerEvent& MouseEvent)
 {
+	const FName CollectionName = CollectionListFilter.Pin()->CollectionName;
 	if (CollectionName != UObjectMixerEditorSerializedData::AllCollectionName && CollectionName != NAME_None)
 	{
 		TSharedRef<FCollectionSelectionButtonDragDropOp> OperationFromCollection =
@@ -181,6 +174,8 @@ FReply SCollectionSelectionButton::OnDragDetected(const FGeometry& MyGeometry, c
 
 void SCollectionSelectionButton::OnDragEnter(const FGeometry& MyGeometry, const FDragDropEvent& DragDropEvent)
 {
+	const FName CollectionName = CollectionListFilter.Pin()->CollectionName;
+	
 	if (TSharedPtr<FObjectMixerListRowDragDropOp> OperationFromRow =
 		DragDropEvent.GetOperationAs<FObjectMixerListRowDragDropOp>())
 	{
@@ -197,7 +192,8 @@ void SCollectionSelectionButton::OnDragEnter(const FGeometry& MyGeometry, const
 	if (TSharedPtr<FCollectionSelectionButtonDragDropOp> OperationFromCollection =
 		DragDropEvent.GetOperationAs<FCollectionSelectionButtonDragDropOp>())
 	{
-		if (CollectionName != UObjectMixerEditorSerializedData::AllCollectionName && CollectionName != NAME_None && CollectionName != OperationFromCollection->DraggedItem)
+		if (CollectionName != UObjectMixerEditorSerializedData::AllCollectionName &&
+			CollectionName != NAME_None && CollectionName != OperationFromCollection->DraggedItem)
 		{
 			OperationFromCollection->SetToolTip(
 			   FText::Format(
@@ -233,6 +229,8 @@ FReply SCollectionSelectionButton::OnDrop(const FGeometry& MyGeometry, const FDr
 	bIsPressed = false;
 	bDropIsValid = false;
 	
+	const FName CollectionName = CollectionListFilter.Pin()->CollectionName;
+	
 	if (TSharedPtr<FObjectMixerListRowDragDropOp> OperationFromRow =
 		DragDropEvent.GetOperationAs<FObjectMixerListRowDragDropOp>())
 	{
@@ -258,11 +256,13 @@ FReply SCollectionSelectionButton::OnDrop(const FGeometry& MyGeometry, const FDr
 	if (TSharedPtr<FCollectionSelectionButtonDragDropOp> OperationFromCollection =
 		DragDropEvent.GetOperationAs<FCollectionSelectionButtonDragDropOp>())
 	{
-		if (CollectionName != UObjectMixerEditorSerializedData::AllCollectionName && CollectionName != NAME_None && CollectionName != OperationFromCollection->DraggedItem)
+		if (CollectionName != UObjectMixerEditorSerializedData::AllCollectionName &&
+			CollectionName != NAME_None && CollectionName != OperationFromCollection->DraggedItem)
 		{
 			if (const TSharedPtr<FObjectMixerEditorMainPanel> MainPanelModel = MainPanelPtr.Pin()->GetMainPanelModel().Pin())
 			{
-				MainPanelModel->RequestReorderCollection(OperationFromCollection->DraggedItem, CollectionName);
+				MainPanelModel->RequestReorderCollection(
+					OperationFromCollection->DraggedItem, CollectionName);
 			}
 		}
 	}
@@ -272,7 +272,7 @@ FReply SCollectionSelectionButton::OnDrop(const FGeometry& MyGeometry, const FDr
 
 const FSlateBrush* SCollectionSelectionButton::GetBorderBrush() const
 {
-	if (MainPanelPtr.Pin()->IsCollectionChecked(CollectionName) == ECheckBoxState::Checked)
+	if (GetIsChecked())
 	{
 		if (bIsPressed)
 		{
@@ -301,8 +301,7 @@ const FSlateBrush* SCollectionSelectionButton::GetBorderBrush() const
 
 FSlateColor SCollectionSelectionButton::GetBorderForeground() const
 {
-	if (MainPanelPtr.Pin()->IsCollectionChecked(CollectionName) == ECheckBoxState::Checked ||
-		bIsPressed || IsHovered())
+	if (GetIsChecked() || bIsPressed || IsHovered())
 	{
 		return FStyleColors::White;
 	}
@@ -310,4 +309,9 @@ FSlateColor SCollectionSelectionButton::GetBorderForeground() const
 	return FStyleColors::Foreground;
 }
 
+bool SCollectionSelectionButton::GetIsChecked() const
+{
+	return CollectionListFilter.Pin()->GetIsFilterActive();
+}
+
 #undef LOCTEXT_NAMESPACE
diff --git a/Engine/Plugins/Editor/ObjectMixer/ObjectMixer/Source/ObjectMixer/Public/Views/List/ObjectMixerEditorList.h b/Engine/Plugins/Editor/ObjectMixer/ObjectMixer/Source/ObjectMixer/Public/Views/List/ObjectMixerEditorList.h
--- a/Engine/Plugins/Editor/ObjectMixer/ObjectMixer/Source/ObjectMixer/Public/Views/List/ObjectMixerEditorList.h
+++ b/Engine/Plugins/Editor/ObjectMixer/ObjectMixer/Source/ObjectMixer/Public/Views/List/ObjectMixerEditorList.h
@@ -13,9 +13,7 @@ class OBJECTMIXEREDITOR_API FObjectMixerEditorList : public TSharedFromThis<FObj
 {
 public:
 
-	FObjectMixerEditorList(TSharedRef<FObjectMixerEditorMainPanel, ESPMode::ThreadSafe> InMainPanel)
-	: MainPanelModelPtr(InMainPanel)
-	{}
+	FObjectMixerEditorList(TSharedRef<FObjectMixerEditorMainPanel, ESPMode::ThreadSafe> InMainPanel);
 
 	virtual ~FObjectMixerEditorList();
 
@@ -23,6 +21,9 @@ public:
 	
 	TSharedRef<SWidget> GetOrCreateWidget();
 
+	void OnPreFilterChange();
+	void OnPostFilterChange();
+
 	void ClearList() const;
 
 	/**
diff --git a/Engine/Plugins/Editor/ObjectMixer/ObjectMixer/Source/ObjectMixer/Public/Views/List/ObjectMixerEditorListFilters/IObjectMixerEditorListFilter.h b/Engine/Plugins/Editor/ObjectMixer/ObjectMixer/Source/ObjectMixer/Public/Views/List/ObjectMixerEditorListFilters/IObjectMixerEditorListFilter.h
--- a/Engine/Plugins/Editor/ObjectMixer/ObjectMixer/Source/ObjectMixer/Public/Views/List/ObjectMixerEditorListFilters/IObjectMixerEditorListFilter.h
+++ b/Engine/Plugins/Editor/ObjectMixer/ObjectMixer/Source/ObjectMixer/Public/Views/List/ObjectMixerEditorListFilters/IObjectMixerEditorListFilter.h
@@ -18,15 +18,28 @@ public:
 	
 	virtual ~IObjectMixerEditorListFilter() = default;
 
-	virtual FString GetFilterName() const = 0;
+	virtual FString GetFilterName() const
+	{
+		return "";
+	}
 
-	/** Return true if the filter should be able to be toggled on/off in the Show Options. */
-	virtual bool IsToggleable() const = 0;
+	/** Return true if the filter should be able to be toggled on/off in the Show Options. Can still be toggled by code. */
+	virtual bool IsUserToggleable() const
+	{
+		return false;
+	}
 
 	/** Returns localized text to display for when this filter' active state is defined by a toggle button. */
-	virtual FText GetFilterButtonLabel() const = 0;
+	virtual FText GetFilterButtonLabel() const
+	{
+		return FText::GetEmpty();
+	}
+	
 	/** Returns localized text to display when the user mouses over a toggle button that defines whether this filter is active. */
-	virtual FText GetFilterButtonToolTip() const = 0;
+	virtual FText GetFilterButtonToolTip() const
+	{
+		return FText::GetEmpty();
+	}
 	
 	void SetFilterActive(const bool bNewEnabled)
 	{
@@ -40,7 +53,7 @@ public:
 
 	bool GetIsFilterActive() const
 	{
-		return !IsToggleable() || bIsFilterActive;
+		return bIsFilterActive;
 	}
 
 	void SetFilterMatchType(const EObjectMixerEditorListFilterMatchType InType)
@@ -60,7 +73,7 @@ public:
 
 private:
 
-	bool bIsFilterActive = true;
+	bool bIsFilterActive = false;
 	EObjectMixerEditorListFilterMatchType FilterMatchType = EObjectMixerEditorListFilterMatchType::MatchAny;
 			
 };
diff --git a/Engine/Plugins/Editor/ObjectMixer/ObjectMixer/Source/ObjectMixer/Public/Views/List/ObjectMixerEditorListFilters/ObjectMixerEditorListFilter_Collection.h b/Engine/Plugins/Editor/ObjectMixer/ObjectMixer/Source/ObjectMixer/Public/Views/List/ObjectMixerEditorListFilters/ObjectMixerEditorListFilter_Collection.h
--- a/Engine/Plugins/Editor/ObjectMixer/ObjectMixer/Source/ObjectMixer/Public/Views/List/ObjectMixerEditorListFilters/ObjectMixerEditorListFilter_Collection.h
+++ b/Engine/Plugins/Editor/ObjectMixer/ObjectMixer/Source/ObjectMixer/Public/Views/List/ObjectMixerEditorListFilters/ObjectMixerEditorListFilter_Collection.h
@@ -10,17 +10,18 @@ class OBJECTMIXEREDITOR_API FObjectMixerEditorListFilter_Collection : public IOb
 {
 public:
 
-	FObjectMixerEditorListFilter_Collection()
+	FObjectMixerEditorListFilter_Collection(const FName InCollectionName)
 	{
+		CollectionName = InCollectionName;
 		SetFilterMatchType(EObjectMixerEditorListFilterMatchType::MatchAll);
 	}
 
 	virtual FString GetFilterName() const override
 	{
-		return "ObjectMixerCollectionListFilter";
+		return "ObjectMixerCollectionListFilter_" + CollectionName.ToString();
 	}
 
-	virtual bool IsToggleable() const override
+	virtual bool IsUserToggleable() const override
 	{
 		return false;
 	}
@@ -39,16 +40,14 @@ public:
 	{
 		if (InItem.IsValid())
 		{
-			if (InItem->GetRowType() == FObjectMixerEditorListRow::ContainerObject ||
-				InItem->GetRowType() == FObjectMixerEditorListRow::MatchingObject)
-			{
-				return InItem->IsObjectRefInSelectedCollections();
-			}
+			return InItem->IsObjectRefInCollection(CollectionName);
 		}
 
 		return false;
 	}
 
+	FName CollectionName;
+
 };
 
 #undef LOCTEXT_NAMESPACE
diff --git a/Engine/Plugins/Editor/ObjectMixer/ObjectMixer/Source/ObjectMixer/Public/Views/List/ObjectMixerEditorListRow.h b/Engine/Plugins/Editor/ObjectMixer/ObjectMixer/Source/ObjectMixer/Public/Views/List/ObjectMixerEditorListRow.h
--- a/Engine/Plugins/Editor/ObjectMixer/ObjectMixer/Source/ObjectMixer/Public/Views/List/ObjectMixerEditorListRow.h
+++ b/Engine/Plugins/Editor/ObjectMixer/ObjectMixer/Source/ObjectMixer/Public/Views/List/ObjectMixerEditorListRow.h
@@ -98,7 +98,7 @@ struct OBJECTMIXEREDITOR_API FObjectMixerEditorListRow final : TSharedFromThis<F
 
 	UObjectMixerObjectFilter* GetObjectFilter() const;
 
-	bool IsObjectRefInSelectedCollections() const;
+	bool IsObjectRefInCollection(const FName& CollectionName) const;
 
 	[[nodiscard]] EObjectMixerEditorListRowType GetRowType() const;
 	void SetRowType(EObjectMixerEditorListRowType InNewRowType);
diff --git a/Engine/Plugins/Editor/ObjectMixer/ObjectMixer/Source/ObjectMixer/Public/Views/List/SObjectMixerEditorList.h b/Engine/Plugins/Editor/ObjectMixer/ObjectMixer/Source/ObjectMixer/Public/Views/List/SObjectMixerEditorList.h
--- a/Engine/Plugins/Editor/ObjectMixer/ObjectMixer/Source/ObjectMixer/Public/Views/List/SObjectMixerEditorList.h
+++ b/Engine/Plugins/Editor/ObjectMixer/ObjectMixer/Source/ObjectMixer/Public/Views/List/SObjectMixerEditorList.h
@@ -140,9 +140,13 @@ public:
 	 */
 	EObjectMixerTreeViewMode GetTreeViewMode();
 	
-	const TArray<TSharedRef<IObjectMixerEditorListFilter>>& GetShowFilters();
+	const TArray<TSharedRef<IObjectMixerEditorListFilter>>& GetListFilters();
 	void EvaluateIfRowsPassFilters(const bool bShouldRefreshAfterward = true);
 
+	/** Saves tree item expanded states to be recalled after the tree view is regenerated. */
+	void CacheTreeState(const TArray<TWeakPtr<IObjectMixerEditorListFilter>>& InFilterCombo);
+	void RestoreTreeState(const TArray<TWeakPtr<IObjectMixerEditorListFilter>>& InFilterCombo, const bool bFlushCache = true);
+
 	// Sorting
 
 	const FName& GetActiveSortingColumnName() const
@@ -216,10 +220,7 @@ protected:
 	void SetAllGroupsCollapsed();
 
 	//  Tree View Implementation
-
-	/** Saves tree item expanded states to be recalled after the tree view is regenerated. */
-	void CacheTreeState();
-	void RestoreTreeState(const bool bFlushCache = true);
+	
 	void BuildPerformanceCacheAndGenerateHeaderIfNeeded();
 	void GenerateTreeView();
 	void FindVisibleTreeViewObjects();
@@ -246,7 +247,13 @@ protected:
 		bool bIsSelected = false;
 	};
 
-	TArray<FTreeItemStateCache> TreeItemStateCache;
+	struct FFilterComboToStateCaches
+	{
+		TArray<TWeakPtr<IObjectMixerEditorListFilter>> FilterCombo = {};
+		TArray<FTreeItemStateCache> Caches = {};
+	};
+
+	TArray<FFilterComboToStateCaches> FilterComboToStateCaches;
 	
 	TMap<UObject*, FObjectMixerEditorListRowPtr> ObjectsToRowsCreated;
 
diff --git a/Engine/Plugins/Editor/ObjectMixer/ObjectMixer/Source/ObjectMixer/Public/Views/MainPanel/ObjectMixerEditorMainPanel.h b/Engine/Plugins/Editor/ObjectMixer/ObjectMixer/Source/ObjectMixer/Public/Views/MainPanel/ObjectMixerEditorMainPanel.h
--- a/Engine/Plugins/Editor/ObjectMixer/ObjectMixer/Source/ObjectMixer/Public/Views/MainPanel/ObjectMixerEditorMainPanel.h
+++ b/Engine/Plugins/Editor/ObjectMixer/ObjectMixer/Source/ObjectMixer/Public/Views/MainPanel/ObjectMixerEditorMainPanel.h
@@ -7,6 +7,7 @@
 #include "UObject/StrongObjectPtr.h"
 #include "Widgets/SWidget.h"
 
+class FObjectMixerEditorListFilter_Collection;
 class UObjectMixerEditorSerializedData;
 struct FObjectMixerEditorListRow;
 class IObjectMixerEditorListFilter;
@@ -14,6 +15,9 @@ class UObjectMixerObjectFilter;
 class FObjectMixerEditorList;
 class SObjectMixerEditorMainPanel;
 
+DECLARE_MULTICAST_DELEGATE(FOnPreFilterChange)
+DECLARE_MULTICAST_DELEGATE(FOnPostFilterChange)
+
 class OBJECTMIXEREDITOR_API FObjectMixerEditorMainPanel : public TSharedFromThis<FObjectMixerEditorMainPanel>
 {
 public:
@@ -104,7 +108,8 @@ public:
 		return ReturnValue;
 	}
 
-	const TArray<TSharedRef<IObjectMixerEditorListFilter>>& GetShowFilters() const;
+	const TArray<TSharedRef<IObjectMixerEditorListFilter>>& GetListFilters() const;
+	TArray<TWeakPtr<IObjectMixerEditorListFilter>> GetWeakActiveListFiltersSortedByName() const;
 
 	/**
 	 * Get the rows that have solo visibility. All other rows should be set to temporarily invisible in editor.
@@ -178,7 +183,10 @@ public:
 	/**
 	 * Returns the collections selected by the user. If the set is empty, consider "All" collections to be selected.
 	 */
-	const TSet<FName>& GetCurrentCollectionSelection() const;
+	TSet<TSharedRef<FObjectMixerEditorListFilter_Collection>> GetCurrentCollectionSelection() const;
+
+	FOnPreFilterChange OnPreFilterChange;
+	FOnPostFilterChange OnPostFilterChange;
 
 private:
 
diff --git a/Engine/Plugins/Editor/ObjectMixer/ObjectMixer/Source/ObjectMixer/Public/Views/MainPanel/SObjectMixerEditorMainPanel.h b/Engine/Plugins/Editor/ObjectMixer/ObjectMixer/Source/ObjectMixer/Public/Views/MainPanel/SObjectMixerEditorMainPanel.h
--- a/Engine/Plugins/Editor/ObjectMixer/ObjectMixer/Source/ObjectMixer/Public/Views/MainPanel/SObjectMixerEditorMainPanel.h
+++ b/Engine/Plugins/Editor/ObjectMixer/ObjectMixer/Source/ObjectMixer/Public/Views/MainPanel/SObjectMixerEditorMainPanel.h
@@ -2,10 +2,12 @@
 
 #pragma once
 
+#include "ObjectMixerEditorSerializedData.h"
 #include "ObjectFilter/ObjectMixerEditorObjectFilter.h"
 #include "Styling/SlateTypes.h"
 #include "Widgets/SCompoundWidget.h"
 
+class FObjectMixerEditorListFilter_Collection;
 class FObjectMixerEditorMainPanel;
 
 class OBJECTMIXEREDITOR_API SObjectMixerEditorMainPanel final : public SCompoundWidget
@@ -41,12 +43,15 @@ public:
 
 	void ToggleFilterActive(const FString& FilterName);
 
-	const TArray<TSharedRef<class IObjectMixerEditorListFilter>>& GetShowFilters()
+	/** Get the filters that affect list item visibility. Distinct from Object Filters. */
+	const TArray<TSharedRef<class IObjectMixerEditorListFilter>>& GetListFilters()
 	{
-		return ShowFilters;	
+		return ListFilters;	
 	}
 
-	const TSet<FName>& GetCurrentCollectionSelection();
+	TArray<TWeakPtr<IObjectMixerEditorListFilter>> GetWeakActiveListFiltersSortedByName();
+
+	TSet<TSharedRef<FObjectMixerEditorListFilter_Collection>> GetCurrentCollectionSelection();
 	
 	void RebuildCollectionSelector();
 
@@ -55,8 +60,8 @@ public:
 	bool RequestRenameCollection(const FName& CollectionNameToRename, const FName& NewCollectionName);
 	bool DoesCollectionExist(const FName& CollectionName) const;
 	
-	void OnCollectionCheckedStateChanged(ECheckBoxState State, FName CollectionName);
-	ECheckBoxState IsCollectionChecked(FName Section) const;
+	void OnCollectionCheckedStateChanged(bool bShouldBeChecked, FName CollectionName);
+	ECheckBoxState IsCollectionChecked(FName CollectionName) const;
 
 	virtual ~SObjectMixerEditorMainPanel() override;
 
@@ -65,7 +70,8 @@ private:
 	/** A reference to the struct that controls this widget */
 	TWeakPtr<FObjectMixerEditorMainPanel> MainPanelModel;
 
-	TArray<TSharedRef<IObjectMixerEditorListFilter>> ShowFilters;
+	/** Filters that affect list item visibility. Distinct from Object Filters. */
+	TArray<TSharedRef<IObjectMixerEditorListFilter>> ListFilters;
 
 	TSharedPtr<class SSearchBox> SearchBoxPtr;
 	TSharedPtr<class SComboButton> ViewOptionsComboButton;
@@ -81,10 +87,10 @@ private:
 	// User Collections
 	
 	TSharedPtr<class SWrapBox> CollectionSelectorBox;
-	TSet<FName> CurrentCollectionSelection;
 
-	void ResetCurrentCollectionSelection()
-	{
-		CurrentCollectionSelection.Reset();
-	}
+	/** Remove all collection filters */
+	void ResetCollectionFilters();
+
+	/** Disable all collection filters except CollectionToEnableName */
+	void SetSingleCollectionSelection(const FName& CollectionToEnableName = UObjectMixerEditorSerializedData::AllCollectionName);
 };
diff --git a/Engine/Plugins/Editor/ObjectMixer/ObjectMixer/Source/ObjectMixer/Public/Views/Widgets/SCollectionSelectionButton.h b/Engine/Plugins/Editor/ObjectMixer/ObjectMixer/Source/ObjectMixer/Public/Views/Widgets/SCollectionSelectionButton.h
--- a/Engine/Plugins/Editor/ObjectMixer/ObjectMixer/Source/ObjectMixer/Public/Views/Widgets/SCollectionSelectionButton.h
+++ b/Engine/Plugins/Editor/ObjectMixer/ObjectMixer/Source/ObjectMixer/Public/Views/Widgets/SCollectionSelectionButton.h
@@ -10,6 +10,8 @@
 
 class SInlineEditableTextBlock;
 class SMenuAnchor;
+class SObjectMixerEditorMainPanel;
+class FObjectMixerEditorListFilter_Collection;
 
 class SCollectionSelectionButton final : public SCompoundWidget
 {
@@ -20,7 +22,10 @@ public:
 
 	SLATE_END_ARGS()
 	
-	void Construct(const FArguments& InArgs, const TSharedRef<class SObjectMixerEditorMainPanel> MainPanelWidget, const FName& InCollectionName);
+	void Construct(
+		const FArguments& InArgs,
+		const TSharedRef<SObjectMixerEditorMainPanel> MainPanelWidget,
+		const TSharedRef<FObjectMixerEditorListFilter_Collection> InCollectionListFilter);
 
 	TSharedRef<SWidget> GetContextMenu() const;
 
@@ -41,8 +46,16 @@ public:
 	const FSlateBrush* GetBorderBrush() const;
 	FSlateColor GetBorderForeground() const;
 
+	TWeakPtr<FObjectMixerEditorListFilter_Collection> GetCollectionListFilter()
+	{
+		return  CollectionListFilter;
+	}
+
+	/** Whether the associated list filter is active */
+	bool GetIsChecked() const;
+
 	TWeakPtr<SObjectMixerEditorMainPanel> MainPanelPtr;
-	FName CollectionName = NAME_None;
+	TWeakPtr<FObjectMixerEditorListFilter_Collection> CollectionListFilter;
 
 private:
 
