Object Mixer:

-Solo multi-light selection

diff --git a/Engine/Plugins/Editor/ObjectMixer/ObjectMixer/Source/ObjectMixer/Private/Views/List/ObjectMixerEditorList.cpp b/Engine/Plugins/Editor/ObjectMixer/ObjectMixer/Source/ObjectMixer/Private/Views/List/ObjectMixerEditorList.cpp
--- a/Engine/Plugins/Editor/ObjectMixer/ObjectMixer/Source/ObjectMixer/Private/Views/List/ObjectMixerEditorList.cpp
+++ b/Engine/Plugins/Editor/ObjectMixer/ObjectMixer/Source/ObjectMixer/Private/Views/List/ObjectMixerEditorList.cpp
@@ -112,27 +112,43 @@ void FObjectMixerEditorList::EvaluateIfRowsPassFilters(const bool bShouldRefresh
 	}
 }
 
-TWeakPtr<FObjectMixerEditorMainPanel> FObjectMixerEditorList::GetMainPanelModel()
+TSet<TWeakPtr<FObjectMixerEditorListRow>> FObjectMixerEditorList::GetSoloRows() const
 {
-	return MainPanelModelPtr;
+	if (ListWidget.IsValid())
+	{
+		return ListWidget->GetSoloRows();
+	}
+	
+	return {};
 }
 
-TSet<TWeakPtr<FObjectMixerEditorListRow>> FObjectMixerEditorList::GetSoloRows()
+void FObjectMixerEditorList::ClearSoloRows()
 {
-	return GetMainPanelModel().Pin()->GetSoloRows();
+	if (ListWidget.IsValid())
+	{
+		ListWidget->ClearSoloRows();
+	}
 }
 
-void FObjectMixerEditorList::AddSoloRow(TSharedRef<FObjectMixerEditorListRow> InRow)
+bool FObjectMixerEditorList::IsListInSoloState() const
 {
-	GetMainPanelModel().Pin()->AddSoloRow(InRow);
+	if (ListWidget.IsValid())
+	{
+		return ListWidget->IsListInSoloState();
+	}
+	
+	return false;
 }
 
-void FObjectMixerEditorList::RemoveSoloRow(TSharedRef<FObjectMixerEditorListRow> InRow)
+void FObjectMixerEditorList::EvaluateAndSetEditorVisibilityPerRow()
 {
-	GetMainPanelModel().Pin()->RemoveSoloRow(InRow);
+	if (ListWidget.IsValid())
+	{
+		ListWidget->EvaluateAndSetEditorVisibilityPerRow();
+	}
 }
 
-void FObjectMixerEditorList::ClearSoloRows()
+TWeakPtr<FObjectMixerEditorMainPanel> FObjectMixerEditorList::GetMainPanelModel()
 {
-	GetMainPanelModel().Pin()->ClearSoloRows();
+	return MainPanelModelPtr;
 }
diff --git a/Engine/Plugins/Editor/ObjectMixer/ObjectMixer/Source/ObjectMixer/Private/Views/List/ObjectMixerEditorListRow.cpp b/Engine/Plugins/Editor/ObjectMixer/ObjectMixer/Source/ObjectMixer/Private/Views/List/ObjectMixerEditorListRow.cpp
--- a/Engine/Plugins/Editor/ObjectMixer/ObjectMixer/Source/ObjectMixer/Private/Views/List/ObjectMixerEditorListRow.cpp
+++ b/Engine/Plugins/Editor/ObjectMixer/ObjectMixer/Source/ObjectMixer/Private/Views/List/ObjectMixerEditorListRow.cpp
@@ -327,14 +327,27 @@ EVisibility FObjectMixerEditorListRow::GetDesiredRowWidgetVisibility() const
 
 bool FObjectMixerEditorListRow::HasVisibleChildRowWidgets() const
 {
-	if (ChildRows.Num() == 0)
+	for (const FObjectMixerEditorListRowPtr& ChildRow : ChildRows)
 	{
-		return false;
+		if (ChildRow->ShouldRowWidgetBeVisible())
+		{
+			return true;
+		}
 	}
 
-	for (const auto& ChildRow : ChildRows)
+	return false;
+}
+
+bool FObjectMixerEditorListRow::HasAtLeastOneChildThatIsNotSolo(const bool bRecursive) const
+{
+	for (const FObjectMixerEditorListRowPtr& ChildRow : ChildRows)
 	{
-		if (ChildRow->ShouldRowWidgetBeVisible())
+		if (!ChildRow->GetRowSoloState())
+		{
+			return true;
+		}
+
+		if (bRecursive && ChildRow->HasAtLeastOneChildThatIsNotSolo())
 		{
 			return true;
 		}
@@ -418,14 +431,14 @@ const FSlateBrush* FObjectMixerEditorListRow::GetObjectIconBrush()
 	return nullptr;
 }
 
-bool FObjectMixerEditorListRow::GetObjectVisibility()
+bool FObjectMixerEditorListRow::GetCurrentEditorObjectVisibility()
 {
 	if (GetRowType() == EObjectMixerEditorListRowType::Folder)
 	{
 		// If any child returns true, the folder returns true
 		for (const FObjectMixerEditorListRowPtr& Child : GetChildRows())
 		{
-			if (Child->GetObjectVisibility())
+			if (Child->GetCurrentEditorObjectVisibility())
 			{
 				return true;
 			}
@@ -443,7 +456,7 @@ bool FObjectMixerEditorListRow::GetObjectVisibility()
 	return false;
 }
 
-void FObjectMixerEditorListRow::SetObjectVisibility(const bool bNewIsVisible, const bool bIsRecursive)
+void FObjectMixerEditorListRow::SetCurrentEditorObjectVisibility(const bool bNewIsVisible, const bool bIsRecursive)
 {
 	if (const UObjectMixerObjectFilter* Filter = GetObjectFilter())
 	{
@@ -453,27 +466,30 @@ void FObjectMixerEditorListRow::SetObjectVisibility(const bool bNewIsVisible, co
 		{
 			for (const FObjectMixerEditorListRowPtr& Child : GetChildRows())
 			{
-				Child->SetObjectVisibility(bNewIsVisible, true);
+				Child->SetCurrentEditorObjectVisibility(bNewIsVisible, true);
 			}
 		}
 	}
 }
 
-bool FObjectMixerEditorListRow::IsThisRowSolo()
+bool FObjectMixerEditorListRow::IsUserSetHiddenInEditor() const
+{
+	return VisibilityRules.bShouldBeHiddenInEditor;
+}
+
+void FObjectMixerEditorListRow::SetUserHiddenInEditor(const bool bNewHidden)
+{
+	VisibilityRules.bShouldBeHiddenInEditor = bNewHidden;
+}
+
+bool FObjectMixerEditorListRow::GetRowSoloState() const
 {
-	return GetListViewPtr().Pin()->GetSoloRows().Contains(SharedThis(this));
+	return VisibilityRules.bShouldBeSolo;
 }
 
 void FObjectMixerEditorListRow::SetRowSoloState(const bool bNewSolo)
 {
-	if (bNewSolo)
-	{
-		GetListViewPtr().Pin()->AddSoloRow(SharedThis(this));
-	}
-	else
-	{
-		GetListViewPtr().Pin()->RemoveSoloRow(SharedThis(this));
-	}
+	VisibilityRules.bShouldBeSolo = bNewSolo;
 }
 
 void FObjectMixerEditorListRow::ClearSoloRows() const
@@ -485,3 +501,13 @@ FObjectMixerEditorListRowPtr FObjectMixerEditorListRow::GetAsShared()
 {
 	return SharedThis(this);
 }
+
+const FObjectMixerEditorListRow::FTransientEditorVisibilityRules& FObjectMixerEditorListRow::GetVisibilityRules() const
+{
+	return VisibilityRules;
+}
+
+void FObjectMixerEditorListRow::SetVisibilityRules(const FTransientEditorVisibilityRules& InVisibilityRules)
+{
+	VisibilityRules = InVisibilityRules;
+}
diff --git a/Engine/Plugins/Editor/ObjectMixer/ObjectMixer/Source/ObjectMixer/Private/Views/List/SObjectMixerEditorList.cpp b/Engine/Plugins/Editor/ObjectMixer/ObjectMixer/Source/ObjectMixer/Private/Views/List/SObjectMixerEditorList.cpp
--- a/Engine/Plugins/Editor/ObjectMixer/ObjectMixer/Source/ObjectMixer/Private/Views/List/SObjectMixerEditorList.cpp
+++ b/Engine/Plugins/Editor/ObjectMixer/ObjectMixer/Source/ObjectMixer/Private/Views/List/SObjectMixerEditorList.cpp
@@ -246,9 +246,9 @@ void SObjectMixerEditorList::SyncEditorSelectionToListSelection()
 
 void SObjectMixerEditorList::SetSelectedTreeViewItemActorsEditorVisible(const bool bNewIsVisible, const bool bIsRecursive)
 {
-	for (const TSharedPtr<FObjectMixerEditorListRow>& SelectedItem : TreeViewPtr->GetSelectedItems())
+	for (const FObjectMixerEditorListRowPtr& SelectedItem : TreeViewPtr->GetSelectedItems())
 	{
-		SelectedItem->SetObjectVisibility(bNewIsVisible, bIsRecursive);
+		SelectedItem->SetCurrentEditorObjectVisibility(bNewIsVisible, bIsRecursive);
 	}
 }
 
@@ -269,24 +269,123 @@ void SObjectMixerEditorList::SetTreeViewItems(const TArray<FObjectMixerEditorLis
 	TreeViewPtr->RequestListRefresh();
 }
 
-TSet<TWeakPtr<FObjectMixerEditorListRow>> SObjectMixerEditorList::GetSoloRows()
+TSet<TWeakPtr<FObjectMixerEditorListRow>> SObjectMixerEditorList::GetSoloRows() const
 {
-	return GetListModelPtr().Pin()->GetSoloRows();
+	TSet<TWeakPtr<FObjectMixerEditorListRow>> SoloRows;
+	
+	using LambdaType = void(*)(const FObjectMixerEditorListRowPtr&, TSet<TWeakPtr<FObjectMixerEditorListRow>>&);
+	static LambdaType GetSoloRowsRecursively =
+		[](const FObjectMixerEditorListRowPtr& InRow, TSet<TWeakPtr<FObjectMixerEditorListRow>>& SoloRows)
+	{
+		if (InRow.IsValid())
+		{
+			if (InRow->GetRowSoloState())
+			{
+				SoloRows.Add(InRow);
+			}
+			
+			for (const FObjectMixerEditorListRowPtr& Child : InRow->GetChildRows())
+			{
+				GetSoloRowsRecursively(Child, SoloRows);
+			}
+		}
+	};
+
+	for (const FObjectMixerEditorListRowPtr& RootRow : TreeViewRootObjects)
+	{
+		GetSoloRowsRecursively(RootRow, SoloRows);
+	}
+
+	return SoloRows;
 }
 
-void SObjectMixerEditorList::AddSoloRow(TSharedRef<FObjectMixerEditorListRow> InRow)
+void SObjectMixerEditorList::ClearSoloRows()
 {
-	GetListModelPtr().Pin()->AddSoloRow(InRow);
+	using LambdaType = void(*)(const FObjectMixerEditorListRowPtr&);
+	static LambdaType ClearSoloRowsRecursively = [](const FObjectMixerEditorListRowPtr& InRow)
+	{
+		if (InRow.IsValid())
+		{
+			InRow->SetRowSoloState(false);
+			
+			for (const FObjectMixerEditorListRowPtr& Child : InRow->GetChildRows())
+			{
+				ClearSoloRowsRecursively(Child);
+			}
+		}
+	};
+
+	for (const FObjectMixerEditorListRowPtr& RootRow : TreeViewRootObjects)
+	{
+		ClearSoloRowsRecursively(RootRow);
+	}
 }
 
-void SObjectMixerEditorList::RemoveSoloRow(TSharedRef<FObjectMixerEditorListRow> InRow)
+bool SObjectMixerEditorList::IsListInSoloState() const
 {
-	GetListModelPtr().Pin()->RemoveSoloRow(InRow);
+	using LambdaType = bool(*)(const FObjectMixerEditorListRowPtr&);
+	static LambdaType GetSoloRowsRecursively = [](const FObjectMixerEditorListRowPtr& InRow)
+	{
+		if (InRow.IsValid())
+		{
+			if (InRow->GetRowSoloState())
+			{
+				return true;
+			}
+			
+			for (const FObjectMixerEditorListRowPtr& Child : InRow->GetChildRows())
+			{
+				if (GetSoloRowsRecursively(Child))
+				{
+					return true;
+				}
+			}
+		}
+
+		return false;
+	};
+
+	for (const FObjectMixerEditorListRowPtr& RootRow : TreeViewRootObjects)
+	{
+		if (GetSoloRowsRecursively(RootRow))
+		{
+			return true;
+		}
+	}
+
+	return false;
 }
 
-void SObjectMixerEditorList::ClearSoloRows()
+void SObjectMixerEditorList::EvaluateAndSetEditorVisibilityPerRow()
 {
-	GetListModelPtr().Pin()->ClearSoloRows();
+	using LambdaType = void(*)(const FObjectMixerEditorListRowPtr&, const bool);
+	static LambdaType EvaluateEditorVisibilityRecursively =
+		[](const FObjectMixerEditorListRowPtr& InRow, const bool bIsListInSoloState)
+	{
+		if (InRow.IsValid())
+		{
+			bool bShouldBeVisible = !InRow->IsUserSetHiddenInEditor();
+
+			if (bIsListInSoloState)
+			{
+				bShouldBeVisible = InRow->GetRowSoloState();
+			}
+
+			InRow->SetCurrentEditorObjectVisibility(bShouldBeVisible);
+			
+			for (const FObjectMixerEditorListRowPtr& Child : InRow->GetChildRows())
+			{
+				EvaluateEditorVisibilityRecursively(Child, bIsListInSoloState);
+			}
+		}
+	};
+	
+	const bool bIsListInSoloState = IsListInSoloState();
+
+	for (const FObjectMixerEditorListRowPtr& RootRow : TreeViewRootObjects)
+	{
+		EvaluateEditorVisibilityRecursively(RootRow, bIsListInSoloState);
+	}
 }
 
 FText SObjectMixerEditorList::GetSearchTextFromSearchInputField() const
@@ -315,7 +414,7 @@ void SObjectMixerEditorList::ExecuteListViewSearchOnAllRows(
 	// unquoted search equivalent to a match-any-of search
 	SearchString.ParseIntoArray(Tokens, TEXT("|"), true);
 	
-	for (const TSharedPtr<FObjectMixerEditorListRow>& ChildRow : TreeViewRootObjects)
+	for (const FObjectMixerEditorListRowPtr& ChildRow : TreeViewRootObjects)
 	{
 		if (!ensure(ChildRow.IsValid()))
 		{
@@ -339,7 +438,7 @@ bool SObjectMixerEditorList::DoesTreeViewHaveVisibleChildren() const
 {
 	if (TreeViewPtr.IsValid())
 	{
-		for (const TSharedPtr<FObjectMixerEditorListRow>& Header : TreeViewRootObjects)
+		for (const FObjectMixerEditorListRowPtr& Header : TreeViewRootObjects)
 		{
 			const EVisibility HeaderVisibility = Header->GetDesiredRowWidgetVisibility();
 			
@@ -353,7 +452,7 @@ bool SObjectMixerEditorList::DoesTreeViewHaveVisibleChildren() const
 	return false;
 }
 
-bool SObjectMixerEditorList::IsTreeViewItemExpanded(const TSharedPtr<FObjectMixerEditorListRow>& Row) const
+bool SObjectMixerEditorList::IsTreeViewItemExpanded(const FObjectMixerEditorListRowPtr& Row) const
 {
 	if (TreeViewPtr.IsValid())
 	{
@@ -363,7 +462,7 @@ bool SObjectMixerEditorList::IsTreeViewItemExpanded(const TSharedPtr<FObjectMixe
 	return false;
 }
 
-void SObjectMixerEditorList::SetTreeViewItemExpanded(const TSharedPtr<FObjectMixerEditorListRow>& RowToExpand, const bool bNewExpansion) const
+void SObjectMixerEditorList::SetTreeViewItemExpanded(const FObjectMixerEditorListRowPtr& RowToExpand, const bool bNewExpansion) const
 {
 	if (TreeViewPtr.IsValid())
 	{
@@ -880,34 +979,36 @@ void SObjectMixerEditorList::SetAllGroupsCollapsed()
 
 void SObjectMixerEditorList::CacheTreeState(const TArray<TWeakPtr<IObjectMixerEditorListFilter>>& InFilterCombo)
 {
-	struct Local
+	using LambdaType = void(*)(const TArray<FObjectMixerEditorListRowPtr>&,
+			TArray<FTreeItemStateCache>*,
+			TSharedPtr<STreeView<FObjectMixerEditorListRowPtr>>);
+	
+	static LambdaType RecursivelyCacheTreeState = [](
+		const TArray<FObjectMixerEditorListRowPtr>& InObjects,
+		TArray<FTreeItemStateCache>* TreeItemStateCache,
+		TSharedPtr<STreeView<FObjectMixerEditorListRowPtr>> TreeViewPtr)
 	{
-		static void RecursivelyCacheTreeState(
-			const TArray<FObjectMixerEditorListRowPtr>& InObjects,
-			TArray<FTreeItemStateCache>* TreeItemStateCache,
-			TSharedPtr<STreeView<FObjectMixerEditorListRowPtr>> TreeViewPtr)
+		for (const FObjectMixerEditorListRowPtr& TreeViewItem : InObjects)
 		{
-			for (const TSharedPtr<FObjectMixerEditorListRow>& TreeViewItem : InObjects)
-			{
-				UObject* RowObject = TreeViewItem->GetObject();
-				const FString RowName =
-					TreeViewItem->GetRowType() == FObjectMixerEditorListRow::Folder ?
-						TreeViewItem->GetFolderPath().ToString() : TreeViewItem->GetDisplayName().ToString();
+			UObject* RowObject = TreeViewItem->GetObject();
+			const FString RowName =
+				TreeViewItem->GetRowType() == FObjectMixerEditorListRow::Folder ?
+					TreeViewItem->GetFolderPath().ToString() : TreeViewItem->GetDisplayName().ToString();
 
-				if (!RowName.IsEmpty())
-				{
-					TreeItemStateCache->Add(
-						{
-							RowObject ? RowObject->GetUniqueID() : -1,
-							RowName,
-							TreeViewPtr->IsItemExpanded(TreeViewItem), 
-							TreeViewPtr->IsItemSelected(TreeViewItem)
-						}
-					);
-				}
-
-				RecursivelyCacheTreeState(TreeViewItem->GetChildRows(), TreeItemStateCache, TreeViewPtr);
+			if (!RowName.IsEmpty())
+			{
+				TreeItemStateCache->Add(
+					{
+						RowObject ? RowObject->GetUniqueID() : -1,
+						RowName,
+						TreeViewPtr->IsItemExpanded(TreeViewItem), 
+						TreeViewPtr->IsItemSelected(TreeViewItem),
+						TreeViewItem->GetVisibilityRules()
+					}
+				);
 			}
+
+			RecursivelyCacheTreeState(TreeViewItem->GetChildRows(), TreeItemStateCache, TreeViewPtr);
 		}
 	};
 
@@ -921,59 +1022,63 @@ void SObjectMixerEditorList::CacheTreeState(const TArray<TWeakPtr<IObjectMixerEd
 
 	if (Match)
 	{
-		Local::RecursivelyCacheTreeState(TreeViewRootObjects, &Match->Caches, TreeViewPtr);
+		RecursivelyCacheTreeState(TreeViewRootObjects, &Match->Caches, TreeViewPtr);
 	}
 	else
 	{
 		TArray<FTreeItemStateCache> NewCache;
-		Local::RecursivelyCacheTreeState(TreeViewRootObjects, &NewCache, TreeViewPtr);
+		RecursivelyCacheTreeState(TreeViewRootObjects, &NewCache, TreeViewPtr);
 		FilterComboToStateCaches.Add({InFilterCombo, MoveTemp(NewCache)});
 	}
 }
 
 void SObjectMixerEditorList::RestoreTreeState(const TArray<TWeakPtr<IObjectMixerEditorListFilter>>& InFilterCombo, const bool bFlushCache)
 {
-	struct Local
+	using LambdaType = void(*)(const TArray<FObjectMixerEditorListRowPtr>&,
+			TArray<FTreeItemStateCache>*,
+			TSharedPtr<STreeView<FObjectMixerEditorListRowPtr>>,
+			const bool);
+	
+	static LambdaType RecursivelyRestoreTreeState = [](
+		const TArray<FObjectMixerEditorListRowPtr>& InObjects,
+		TArray<FTreeItemStateCache>* TreeItemStateCache,
+		TSharedPtr<STreeView<FObjectMixerEditorListRowPtr>> TreeViewPtr,
+		const bool bExpandByDefault)
 	{
-		static void RecursivelyRestoreTreeState(
-			const TArray<FObjectMixerEditorListRowPtr>& InObjects,
-			TArray<FTreeItemStateCache>* TreeItemStateCache,
-			TSharedPtr<STreeView<FObjectMixerEditorListRowPtr>> TreeViewPtr,
-			const bool bExpandByDefault)
+		for (const FObjectMixerEditorListRowPtr& TreeViewItem : InObjects)
 		{
-			for (const TSharedPtr<FObjectMixerEditorListRow>& TreeViewItem : InObjects)
-			{
-				if (const FTreeItemStateCache* StateCachePtr = Algo::FindByPredicate(
-					*TreeItemStateCache,
-					[TreeViewItem](const FTreeItemStateCache& Other)
+			if (const FTreeItemStateCache* StateCachePtr = Algo::FindByPredicate(
+				*TreeItemStateCache,
+				[TreeViewItem](const FTreeItemStateCache& Other)
+				{
+					if (Other.UniqueId != -1)
 					{
-						if (Other.UniqueId != -1)
+						if (const UObject* RowObject = TreeViewItem->GetObject())
 						{
-							if (const UObject* RowObject = TreeViewItem->GetObject())
-							{
-								return Other.UniqueId == RowObject->GetUniqueID();
-							}
+							return Other.UniqueId == RowObject->GetUniqueID();
 						}
+					}
 
-						const FString RowName =
-							TreeViewItem->GetRowType() == FObjectMixerEditorListRow::Folder ?
-								TreeViewItem->GetFolderPath().ToString() : TreeViewItem->GetDisplayName().ToString();
+					const FString RowName =
+						TreeViewItem->GetRowType() == FObjectMixerEditorListRow::Folder ?
+							TreeViewItem->GetFolderPath().ToString() : TreeViewItem->GetDisplayName().ToString();
 
-						return Other.RowName.Equals(RowName);
-					}
-				))
-				{
-					TreeViewPtr->SetItemExpansion(TreeViewItem, StateCachePtr->bIsExpanded);
-					TreeViewPtr->SetItemSelection(TreeViewItem, StateCachePtr->bIsSelected);
-				}
-				else
-				{
-					TreeViewPtr->SetItemExpansion(TreeViewItem, bExpandByDefault);
+					return Other.RowName.Equals(RowName);
 				}
-
-				RecursivelyRestoreTreeState(
-					TreeViewItem->GetChildRows(), TreeItemStateCache, TreeViewPtr, bExpandByDefault);
+			))
+			{
+				TreeViewPtr->SetItemExpansion(TreeViewItem, StateCachePtr->bIsExpanded);
+				TreeViewPtr->SetItemSelection(TreeViewItem, StateCachePtr->bIsSelected);
+				
+				TreeViewItem->SetVisibilityRules(StateCachePtr->VisibilityRules);
 			}
+			else
+			{
+				TreeViewPtr->SetItemExpansion(TreeViewItem, bExpandByDefault);
+			}
+
+			RecursivelyRestoreTreeState(
+				TreeViewItem->GetChildRows(), TreeItemStateCache, TreeViewPtr, bExpandByDefault);
 		}
 	};
 	
@@ -987,7 +1092,7 @@ void SObjectMixerEditorList::RestoreTreeState(const TArray<TWeakPtr<IObjectMixer
 	{
 		if (FilterComboToStateCaches[CachesItr].FilterCombo == InFilterCombo)
 		{
-			Local::RecursivelyRestoreTreeState(
+			RecursivelyRestoreTreeState(
 			  TreeViewRootObjects, &FilterComboToStateCaches[CachesItr].Caches, TreeViewPtr, bExpandByDefault);
 
 			if (bFlushCache)
@@ -1179,7 +1284,7 @@ void CreateFolderRowsForTopLevelRow(
 
 		while (!BaseActorFolder.IsNone())
 		{
-			TSharedPtr<FObjectMixerEditorListRow> FolderRow;
+			FObjectMixerEditorListRowPtr FolderRow;
 			
 			if (const FObjectMixerEditorListRowPtr* Match = FolderMap.Find(BaseActorFolder.GetPath()))
 			{
@@ -1287,7 +1392,7 @@ void SObjectMixerEditorList::FindVisibleTreeViewObjects()
 {
 	VisibleTreeViewObjects.Empty();
 
-	for (const TSharedPtr<FObjectMixerEditorListRow>& Row : TreeViewRootObjects)
+	for (const FObjectMixerEditorListRowPtr& Row : TreeViewRootObjects)
 	{
 		if (Row->ShouldRowWidgetBeVisible())
 		{
@@ -1316,7 +1421,7 @@ void SObjectMixerEditorList::SelectedTreeItemsToSelectedInLevelEditor()
 		}
 		else
 		{
-			for (const TSharedPtr<FObjectMixerEditorListRow>& SelectedRow : SelectedTreeItems)
+			for (const FObjectMixerEditorListRowPtr& SelectedRow : SelectedTreeItems)
 			{
 				if (SelectedRow->GetRowType() == FObjectMixerEditorListRow::MatchingObject ||
 					SelectedRow->GetRowType() == FObjectMixerEditorListRow::ContainerObject)
diff --git a/Engine/Plugins/Editor/ObjectMixer/ObjectMixer/Source/ObjectMixer/Private/Views/List/SObjectMixerEditorListRow.cpp b/Engine/Plugins/Editor/ObjectMixer/ObjectMixer/Source/ObjectMixer/Private/Views/List/SObjectMixerEditorListRow.cpp
--- a/Engine/Plugins/Editor/ObjectMixer/ObjectMixer/Source/ObjectMixer/Private/Views/List/SObjectMixerEditorListRow.cpp
+++ b/Engine/Plugins/Editor/ObjectMixer/ObjectMixer/Source/ObjectMixer/Private/Views/List/SObjectMixerEditorListRow.cpp
@@ -275,7 +275,7 @@ bool SObjectMixerEditorListRow::IsVisible() const
 {
 	if (const TSharedPtr<FObjectMixerEditorListRow> PinnedItem = Item.Pin())
 	{
-		return PinnedItem->GetObjectVisibility();
+		return PinnedItem->GetCurrentEditorObjectVisibility();
 	}
 
 	return false;
@@ -303,11 +303,12 @@ FSlateColor SObjectMixerEditorListRow::GetVisibilityIconForegroundColor() const
 FSlateColor SObjectMixerEditorListRow::GetSoloIconForegroundColor() const
 {
 	check(Item.IsValid());
+	const FObjectMixerEditorListRowPtr RowPtr = GetHybridChildOrRowItemIfNull();
 
-	const bool bIsSelected = Item.Pin()->GetIsSelected();
+	const bool bIsSelected = RowPtr->GetIsSelected();
 
 	// make the foreground brush transparent if it is not selected, hovered or solo
-	if (!Item.Pin()->IsThisRowSolo() && !bIsHovered && !bIsSelected)
+	if (!RowPtr->GetRowSoloState() && !bIsHovered && !bIsSelected)
 	{
 		return FLinearColor::Transparent;
 	}
@@ -319,6 +320,49 @@ FSlateColor SObjectMixerEditorListRow::GetSoloIconForegroundColor() const
 	return FSlateColor::UseForeground();
 }
 
+void SObjectMixerEditorListRow::OnClickSoloIcon(const FObjectMixerEditorListRowPtr& RowPtr)
+{
+	check (RowPtr.IsValid());
+							
+	if (const TSharedPtr<SObjectMixerEditorList> PinnedListView = RowPtr->GetListViewPtr().Pin();
+		PinnedListView->GetTreeViewItemCount() > 0)
+	{
+		const bool bNewSolo = !RowPtr->GetRowSoloState();
+								
+		using LambdaType = void(*)(const FObjectMixerEditorListRowPtr&, const bool);
+		
+		static LambdaType SetSoloPerRowRecursively =
+			[](const FObjectMixerEditorListRowPtr& RowPtr, const bool bNewSolo)
+		{
+			if (bNewSolo)
+			{
+				RowPtr->SetUserHiddenInEditor(false);
+			}
+			
+			RowPtr->SetRowSoloState(bNewSolo);
+								
+			for (const FObjectMixerEditorListRowPtr& SelectedItem : RowPtr->GetChildRows())
+			{
+				SetSoloPerRowRecursively(SelectedItem, bNewSolo);
+			}
+		};
+								
+		if (PinnedListView->GetSelectedTreeViewItemCount() > 0 && RowPtr->GetIsSelected())
+		{
+			for (const FObjectMixerEditorListRowPtr& SelectedItem : PinnedListView->GetSelectedTreeViewItems())
+			{
+				SetSoloPerRowRecursively(SelectedItem, bNewSolo);
+			}
+		}
+		else
+		{
+			SetSoloPerRowRecursively(RowPtr, bNewSolo);
+		}
+								
+		PinnedListView->EvaluateAndSetEditorVisibilityPerRow();
+	}
+}
+
 const FSlateBrush* SObjectMixerEditorListRow::GetVisibilityBrush() const
 {
 	if (IsVisible())
@@ -332,9 +376,19 @@ const FSlateBrush* SObjectMixerEditorListRow::GetVisibilityBrush() const
 const FSlateBrush* SObjectMixerEditorListRow::GetSoloBrush() const
 {
 	check(Item.IsValid());
+	const FObjectMixerEditorListRowPtr RowPtr = GetHybridChildOrRowItemIfNull();
 
-	// make the foreground brush transparent if it is not selected, hovered or solo
-	if (Item.Pin()->IsThisRowSolo())
+	if (RowPtr->GetRowType() == FObjectMixerEditorListRow::Folder)
+	{
+		if (RowPtr->HasAtLeastOneChildThatIsNotSolo())
+		{
+			return SoloOffHoveredBrush;
+		}
+
+		return SoloOnBrush;
+	}
+	
+	if (RowPtr->GetRowSoloState())
 	{
 		return SoloOnBrush;
 	}
@@ -342,6 +396,52 @@ const FSlateBrush* SObjectMixerEditorListRow::GetSoloBrush() const
 	return SoloOffHoveredBrush;
 }
 
+void SObjectMixerEditorListRow::OnClickVisibilityIcon(const FObjectMixerEditorListRowPtr& RowPtr)
+{
+	check (RowPtr.IsValid());
+							
+	if (const TSharedPtr<SObjectMixerEditorList> PinnedListView = RowPtr->GetListViewPtr().Pin();
+		PinnedListView->GetTreeViewItemCount() > 0)
+	{
+		const bool bNewHidden = !RowPtr->IsUserSetHiddenInEditor();
+		const bool bIsListInSoloState = PinnedListView->IsListInSoloState();
+								
+		using LambdaType = void(*)(const FObjectMixerEditorListRowPtr&, const bool, const bool);
+		
+		static LambdaType SetVisibilityPerRowRecursively = [](
+			const FObjectMixerEditorListRowPtr& RowPtr, const bool bNewHidden, const bool bIsListInSoloState)
+		{
+			if (bIsListInSoloState)
+			{
+				RowPtr->SetRowSoloState(!RowPtr->GetRowSoloState());
+			}
+			else
+			{
+				RowPtr->SetUserHiddenInEditor(bNewHidden);
+			}
+								
+			for (const FObjectMixerEditorListRowPtr& SelectedItem : RowPtr->GetChildRows())
+			{
+				SetVisibilityPerRowRecursively(SelectedItem, bNewHidden, bIsListInSoloState);
+			}
+		};
+								
+		if (PinnedListView->GetSelectedTreeViewItemCount() > 0 && RowPtr->GetIsSelected())
+		{
+			for (const FObjectMixerEditorListRowPtr& SelectedItem : PinnedListView->GetSelectedTreeViewItems())
+			{
+				SetVisibilityPerRowRecursively(SelectedItem, bNewHidden, bIsListInSoloState);
+			}
+		}
+		else
+		{
+			SetVisibilityPerRowRecursively(RowPtr, bNewHidden, bIsListInSoloState);
+		}
+								
+		PinnedListView->EvaluateAndSetEditorVisibilityPerRow();
+	}
+}
+
 TSharedPtr<SWidget> SObjectMixerEditorListRow::GenerateCells(
 	const FName& InColumnName, const TSharedPtr<FObjectMixerEditorListRow> RowPtr)
 {
@@ -446,24 +546,10 @@ TSharedPtr<SWidget> SObjectMixerEditorListRow::GenerateCells(
 					.ColorAndOpacity(this, &SObjectMixerEditorListRow::GetVisibilityIconForegroundColor)
 					.Image_Raw(this, &SObjectMixerEditorListRow::GetVisibilityBrush)
 					.OnMouseButtonDown_Lambda(
-						[RowPtr] (const FGeometry& MyGeometry, const FPointerEvent& Event)
+						[this] (const FGeometry&, const FPointerEvent&)
 						{
-							check (RowPtr);
-
-							FScopedTransaction Transaction( LOCTEXT("VisibilityChanged", "Object Mixer - Visibility Changed") );
-
-							const bool bIsVisible = RowPtr->GetObjectVisibility();
-							
-							if (const TSharedPtr<SObjectMixerEditorList> PinnedListView = RowPtr->GetListViewPtr().Pin();
-								RowPtr->GetIsSelected() && PinnedListView->GetSelectedTreeViewItemCount() > 0)
-							{
-								PinnedListView->SetSelectedTreeViewItemActorsEditorVisible(!bIsVisible);
-
-								return FReply::Handled();
-							}
-
-							// Set Visibility Recursively
-							RowPtr->SetObjectVisibility(!bIsVisible, true);
+							// Pass in actual row item even if hybrid row
+							OnClickVisibilityIcon(Item.Pin());
 
 							return FReply::Handled();
 						}
@@ -489,33 +575,12 @@ TSharedPtr<SWidget> SObjectMixerEditorListRow::GenerateCells(
 					.ColorAndOpacity(this, &SObjectMixerEditorListRow::GetSoloIconForegroundColor)
 					.Image(this, &SObjectMixerEditorListRow::GetSoloBrush)
 					.OnMouseButtonDown_Lambda(
-						[this] (const FGeometry& MyGeometry, const FPointerEvent& Event)
+						[this] (const FGeometry&, const FPointerEvent&)
 						{
-							check (Item.IsValid());
+							// Pass in actual row item even if hybrid row
+							OnClickSoloIcon(Item.Pin());
 
-							const FObjectMixerEditorListRowPtr PinnedItem = Item.Pin();
-							
-							if (const TSharedPtr<SObjectMixerEditorList> PinnedListView = PinnedItem->GetListViewPtr().Pin();
-								PinnedListView->GetTreeViewItemCount() > 0)
-							{
-								const bool bIsRowSolo = PinnedItem->IsThisRowSolo();
-								for (const TSharedPtr<FObjectMixerEditorListRow>& TreeItem : PinnedListView->GetTreeViewItems())
-								{
-									TreeItem->SetObjectVisibility(bIsRowSolo, true);
-								}
-
-								if (!bIsRowSolo)
-								{
-									PinnedItem->SetObjectVisibility(true, true);
-									PinnedItem->ClearSoloRows();
-								}
-								
-								PinnedItem->SetRowSoloState(!bIsRowSolo);
-
-								return FReply::Handled();
-							}
-
-							return FReply::Unhandled();
+							return FReply::Handled();
 						}
 					)
 				]
@@ -570,36 +635,33 @@ TSharedPtr<SWidget> SObjectMixerEditorListRow::GenerateCells(
 void SObjectMixerEditorListRow::OnPropertyChanged(const FName PropertyName) const
 {
 	check(Item.IsValid());
-
-	struct Local
+	
+	auto SetValueOnSelectedItems = [](
+		const FString& ValueAsString, const TArray<FObjectMixerEditorListRowPtr>& OtherSelectedItems,
+		const FName& PropertyName, const FObjectMixerEditorListRowPtr PinnedItem)
 	{
-		static void SetValueOnSelectedItems(
-			const FString& ValueAsString, const TArray<FObjectMixerEditorListRowPtr>& OtherSelectedItems,
-			const FName& PropertyName, const FObjectMixerEditorListRowPtr PinnedItem)
+		if (!ValueAsString.IsEmpty())
 		{
-			if (!ValueAsString.IsEmpty())
+			FScopedTransaction Transaction(
+				LOCTEXT("OnPropertyChangedTransaction", "Object Mixer - Bulk Edit Selected Row Properties") );
+					
+			for (const TSharedPtr<FObjectMixerEditorListRow>& SelectedRow : OtherSelectedItems)
 			{
-				FScopedTransaction Transaction(
-					LOCTEXT("OnPropertyChangedTransaction", "Object Mixer - Bulk Edit Selected Row Properties") );
-						
-				for (const TSharedPtr<FObjectMixerEditorListRow>& SelectedRow : OtherSelectedItems)
-				{
-					const FObjectMixerEditorListRowPtr SelectedHybridRow = SelectedRow->GetHybridChild();
-					const FObjectMixerEditorListRowPtr RowToUse = SelectedHybridRow.IsValid() ? SelectedHybridRow : SelectedRow;
+				const FObjectMixerEditorListRowPtr SelectedHybridRow = SelectedRow->GetHybridChild();
+				const FObjectMixerEditorListRowPtr RowToUse = SelectedHybridRow.IsValid() ? SelectedHybridRow : SelectedRow;
 
-					if (RowToUse != PinnedItem)
+				if (RowToUse != PinnedItem)
+				{
+					if (const TWeakPtr<IPropertyHandle>* SelectedHandlePtr = RowToUse->PropertyNamesToHandles.Find(PropertyName))
 					{
-						if (const TWeakPtr<IPropertyHandle>* SelectedHandlePtr = RowToUse->PropertyNamesToHandles.Find(PropertyName))
+						if (SelectedHandlePtr->IsValid())
 						{
-							if (SelectedHandlePtr->IsValid())
+							if (UObject* ObjectToModify = RowToUse->GetObject())
 							{
-								if (UObject* ObjectToModify = RowToUse->GetObject())
-								{
-									ObjectToModify->Modify();
-								}
-										
-								SelectedHandlePtr->Pin()->SetValueFromFormattedString(ValueAsString);
+								ObjectToModify->Modify();
 							}
+									
+							SelectedHandlePtr->Pin()->SetValueFromFormattedString(ValueAsString);
 						}
 					}
 				}
@@ -621,7 +683,7 @@ void SObjectMixerEditorListRow::OnPropertyChanged(const FName PropertyName) cons
 						if (TSharedPtr<IPropertyHandle> PinnedHandle = HandlePtr->Pin())
 						{
 							PinnedHandle->GetValueAsFormattedString(ValueAsString);
-							Local::SetValueOnSelectedItems(ValueAsString, OtherSelectedItems, PropertyName, PinnedItem);
+							SetValueOnSelectedItems(ValueAsString, OtherSelectedItems, PropertyName, PinnedItem);
 						}
 					}
 				}
@@ -649,8 +711,9 @@ void SObjectMixerEditorListRow::OnClickBlueprintLink(UBlueprint* AsBlueprint, UO
 FText SObjectMixerEditorListRow::GetHighlightText() const
 {
 	check (Item.IsValid());
+	const FObjectMixerEditorListRowPtr RowPtr = GetHybridChildOrRowItemIfNull();
 
-	if (const TSharedPtr<SObjectMixerEditorList> PinnedListView = Item.Pin()->GetListViewPtr().Pin())
+	if (const TSharedPtr<SObjectMixerEditorList> PinnedListView = RowPtr->GetListViewPtr().Pin())
 	{
 		return PinnedListView->GetSearchTextFromSearchInputField();
 	}
diff --git a/Engine/Plugins/Editor/ObjectMixer/ObjectMixer/Source/ObjectMixer/Private/Views/MainPanel/SObjectMixerEditorMainPanel.cpp b/Engine/Plugins/Editor/ObjectMixer/ObjectMixer/Source/ObjectMixer/Private/Views/MainPanel/SObjectMixerEditorMainPanel.cpp
--- a/Engine/Plugins/Editor/ObjectMixer/ObjectMixer/Source/ObjectMixer/Private/Views/MainPanel/SObjectMixerEditorMainPanel.cpp
+++ b/Engine/Plugins/Editor/ObjectMixer/ObjectMixer/Source/ObjectMixer/Private/Views/MainPanel/SObjectMixerEditorMainPanel.cpp
@@ -406,6 +406,22 @@ TSharedRef<SWidget> SObjectMixerEditorMainPanel::BuildShowOptionsMenu()
 			  })));
 		}
 
+		ShowOptionsMenuBuilder.AddMenuEntry(
+			LOCTEXT("ClearSoloStatesMenuOption","Clear Solo States"), 
+			LOCTEXT("ClearSoloStatesMenuOptionTooltip","Remove the solo state from all rows in this list."),
+			FSlateIcon(),
+			FUIAction(FExecuteAction::CreateLambda(
+				[this]()
+				{
+					if (const TSharedPtr<FObjectMixerEditorList> ListModel = MainPanelModel.Pin()->GetEditorListModel().Pin())
+					{
+						ListModel->ClearSoloRows();
+						ListModel->EvaluateAndSetEditorVisibilityPerRow();
+					}
+				})
+			)
+		);
+
 		ShowOptionsMenuBuilder.AddMenuEntry(
 			LOCTEXT("RebuildListMenuOption","Rebuild List"), 
 			LOCTEXT("RebuildListMenuOptionTooltip","Force the active list to be rebuilt. Useful if the list doesn't automatically refresh."),
diff --git a/Engine/Plugins/Editor/ObjectMixer/ObjectMixer/Source/ObjectMixer/Public/Views/List/ObjectMixerEditorList.h b/Engine/Plugins/Editor/ObjectMixer/ObjectMixer/Source/ObjectMixer/Public/Views/List/ObjectMixerEditorList.h
--- a/Engine/Plugins/Editor/ObjectMixer/ObjectMixer/Source/ObjectMixer/Public/Views/List/ObjectMixerEditorList.h
+++ b/Engine/Plugins/Editor/ObjectMixer/ObjectMixer/Source/ObjectMixer/Public/Views/List/ObjectMixerEditorList.h
@@ -43,14 +43,19 @@ public:
 
 	void EvaluateIfRowsPassFilters(const bool bShouldRefreshAfterward = true) const;
 
-	TWeakPtr<FObjectMixerEditorMainPanel> GetMainPanelModel();
+	TSet<TWeakPtr<FObjectMixerEditorListRow>> GetSoloRows() const;
+	void ClearSoloRows();
 
-	TSet<TWeakPtr<FObjectMixerEditorListRow>> GetSoloRows();
+	/** Returns true if at least one row is set to Solo. */
+	bool IsListInSoloState() const;
 
-	void AddSoloRow(TSharedRef<FObjectMixerEditorListRow> InRow);
-	void RemoveSoloRow(TSharedRef<FObjectMixerEditorListRow> InRow);
+	/**
+	 * Determines whether rows' objects should be temporarily hidden in editor based on each row's visibility rules,
+	 * then sets each object's visibility in editor.
+	 */
+	void EvaluateAndSetEditorVisibilityPerRow();
 
-	void ClearSoloRows();
+	TWeakPtr<FObjectMixerEditorMainPanel> GetMainPanelModel();
 
 private:
 
diff --git a/Engine/Plugins/Editor/ObjectMixer/ObjectMixer/Source/ObjectMixer/Public/Views/List/ObjectMixerEditorListRow.h b/Engine/Plugins/Editor/ObjectMixer/ObjectMixer/Source/ObjectMixer/Public/Views/List/ObjectMixerEditorListRow.h
--- a/Engine/Plugins/Editor/ObjectMixer/ObjectMixer/Source/ObjectMixer/Public/Views/List/ObjectMixerEditorListRow.h
+++ b/Engine/Plugins/Editor/ObjectMixer/ObjectMixer/Source/ObjectMixer/Public/Views/List/ObjectMixerEditorListRow.h
@@ -43,6 +43,22 @@ struct OBJECTMIXEREDITOR_API FObjectMixerEditorListRow final : TSharedFromThis<F
 		MatchingObject // The object that has the properties we wish to affect
 	};
 
+	struct FTransientEditorVisibilityRules
+	{
+		/**
+		 * If true, the user wants the row to be hidden temporarily in the editor.
+		 * This is transient visibility like the eye icon in the Scene Outliner, not the bVisible or bHiddenInGame properties.
+		 */
+		bool bShouldBeHiddenInEditor = false;
+
+		/**
+		 * If true, the user wants the row to have solo visibility. Multiple rows at once can be set to solo.
+		 * Solo rows' objects are exclusively visible,
+		 * so all other objects found in the panel will be invisible while at least one row is in a solo state.
+		 */
+		bool bShouldBeSolo = false;
+	};
+
 	~FObjectMixerEditorListRow();
 
 	bool operator==(const TSharedPtr<FObjectMixerEditorListRow>& Other) const
@@ -149,6 +165,8 @@ struct OBJECTMIXEREDITOR_API FObjectMixerEditorListRow final : TSharedFromThis<F
 
 	[[nodiscard]] bool HasVisibleChildRowWidgets() const;
 
+	[[nodiscard]] bool HasAtLeastOneChildThatIsNotSolo(const bool bRecursive = true) const;
+
 	[[nodiscard]] FText GetDisplayName(const bool bIsHybridRow = false) const;
 
 	[[nodiscard]] const FText& GetDisplayNameOverride() const
@@ -175,10 +193,16 @@ struct OBJECTMIXEREDITOR_API FObjectMixerEditorListRow final : TSharedFromThis<F
 
 	const FSlateBrush* GetObjectIconBrush();
 
-	bool GetObjectVisibility();
-	void SetObjectVisibility(const bool bNewIsVisible, const bool bIsRecursive = false);
+	bool GetCurrentEditorObjectVisibility();
+	void SetCurrentEditorObjectVisibility(const bool bNewIsVisible, const bool bIsRecursive = false);
+	
+	[[nodiscard]] const FTransientEditorVisibilityRules& GetVisibilityRules() const;
+	void SetVisibilityRules(const FTransientEditorVisibilityRules& InVisibilityRules);
 
-	bool IsThisRowSolo();
+	bool IsUserSetHiddenInEditor() const;
+	void SetUserHiddenInEditor(const bool bNewHidden);
+	
+	bool GetRowSoloState() const;
 	void SetRowSoloState(const bool bNewSolo);
 
 	void ClearSoloRows() const;
@@ -196,6 +220,8 @@ private:
 	
 	TWeakPtr<SObjectMixerEditorList> ListViewPtr;
 
+	FTransientEditorVisibilityRules VisibilityRules;
+
 	FText DisplayNameOverride;
 
 	int32 SortOrder = -1;
diff --git a/Engine/Plugins/Editor/ObjectMixer/ObjectMixer/Source/ObjectMixer/Public/Views/List/SObjectMixerEditorList.h b/Engine/Plugins/Editor/ObjectMixer/ObjectMixer/Source/ObjectMixer/Public/Views/List/SObjectMixerEditorList.h
--- a/Engine/Plugins/Editor/ObjectMixer/ObjectMixer/Source/ObjectMixer/Public/Views/List/SObjectMixerEditorList.h
+++ b/Engine/Plugins/Editor/ObjectMixer/ObjectMixer/Source/ObjectMixer/Public/Views/List/SObjectMixerEditorList.h
@@ -119,12 +119,17 @@ public:
 		return TreeViewRootObjects.Num();
 	}
 	
-	TSet<TWeakPtr<FObjectMixerEditorListRow>> GetSoloRows();
+	TSet<TWeakPtr<FObjectMixerEditorListRow>> GetSoloRows() const;
+	void ClearSoloRows();
 
-	void AddSoloRow(TSharedRef<FObjectMixerEditorListRow> InRow);
-	void RemoveSoloRow(TSharedRef<FObjectMixerEditorListRow> InRow);
+	/** Returns true if at least one row is set to Solo. */
+	bool IsListInSoloState() const;
 
-	void ClearSoloRows();
+	/**
+	 * Determines whether rows' objects should be temporarily hidden in editor based on each row's visibility rules,
+	 * then sets each object's visibility in editor.
+	 */
+	void EvaluateAndSetEditorVisibilityPerRow();
 
 	FText GetSearchTextFromSearchInputField() const;
 	FString GetSearchStringFromSearchInputField() const;
@@ -245,6 +250,7 @@ protected:
 		FString RowName = "";
 		bool bIsExpanded = false;
 		bool bIsSelected = false;
+		FObjectMixerEditorListRow::FTransientEditorVisibilityRules VisibilityRules;
 	};
 
 	struct FFilterComboToStateCaches
diff --git a/Engine/Plugins/Editor/ObjectMixer/ObjectMixer/Source/ObjectMixer/Public/Views/List/SObjectMixerEditorListRow.h b/Engine/Plugins/Editor/ObjectMixer/ObjectMixer/Source/ObjectMixer/Public/Views/List/SObjectMixerEditorListRow.h
--- a/Engine/Plugins/Editor/ObjectMixer/ObjectMixer/Source/ObjectMixer/Public/Views/List/SObjectMixerEditorListRow.h
+++ b/Engine/Plugins/Editor/ObjectMixer/ObjectMixer/Source/ObjectMixer/Public/Views/List/SObjectMixerEditorListRow.h
@@ -47,10 +47,12 @@ private:
 
 	FSlateColor GetVisibilityIconForegroundColor() const;
 	FSlateColor GetSoloIconForegroundColor() const;
+	void OnClickSoloIcon(const FObjectMixerEditorListRowPtr& RowPtr);
 	
 	/** Get the brush for this widget */
 	const FSlateBrush* GetVisibilityBrush() const;
 	const FSlateBrush* GetSoloBrush() const;
+	void OnClickVisibilityIcon(const FObjectMixerEditorListRowPtr& RowPtr);
 
 	TSharedPtr<SWidget> GenerateCells(const FName& InColumnName, const TSharedPtr<FObjectMixerEditorListRow> RowPtr);
 
diff --git a/Engine/Plugins/Editor/ObjectMixer/ObjectMixer/Source/ObjectMixer/Public/Views/MainPanel/ObjectMixerEditorMainPanel.h b/Engine/Plugins/Editor/ObjectMixer/ObjectMixer/Source/ObjectMixer/Public/Views/MainPanel/ObjectMixerEditorMainPanel.h
--- a/Engine/Plugins/Editor/ObjectMixer/ObjectMixer/Source/ObjectMixer/Public/Views/MainPanel/ObjectMixerEditorMainPanel.h
+++ b/Engine/Plugins/Editor/ObjectMixer/ObjectMixer/Source/ObjectMixer/Public/Views/MainPanel/ObjectMixerEditorMainPanel.h
@@ -111,38 +111,6 @@ public:
 	const TArray<TSharedRef<IObjectMixerEditorListFilter>>& GetListFilters() const;
 	TArray<TWeakPtr<IObjectMixerEditorListFilter>> GetWeakActiveListFiltersSortedByName() const;
 
-	/**
-	 * Get the rows that have solo visibility. All other rows should be set to temporarily invisible in editor.
-	 */
-	TSet<TWeakPtr<FObjectMixerEditorListRow>> GetSoloRows()
-	{
-		return SoloRows;
-	}
-
-	/**
-	 * Add a row that has solo visibility. This does not set temporary editor invisibility for other rows.
-	 */
-	void AddSoloRow(TSharedRef<FObjectMixerEditorListRow> InRow)
-	{
-		SoloRows.Add(InRow);
-	}
-
-	/**
-	 * Remove a row that does not have solo visibility. This does not set temporary editor invisibility for other rows.
-	 */
-	void RemoveSoloRow(TSharedRef<FObjectMixerEditorListRow> InRow)
-	{
-		SoloRows.Remove(InRow);
-	}
-
-	/**
-	 * Clear the rows that have solo visibility. This does not remove temporary editor invisibility for other rows.
-	 */
-	void ClearSoloRows()
-	{
-		SoloRows.Empty();
-	}
-
 	TSubclassOf<UObjectMixerObjectFilter> GetObjectFilterClass() const
 	{
 		return ObjectFilterClass;
@@ -206,7 +174,5 @@ private:
 	 */
 	EObjectMixerTreeViewMode TreeViewMode = EObjectMixerTreeViewMode::Folders;
 
-	TSet<TWeakPtr<FObjectMixerEditorListRow>> SoloRows = {};
-
 	FName ModuleName = NAME_None;
 };
