SObjectMixerEditorList: Fix shadowed variable compile warning.
#rnx

diff --git a/Engine/Plugins/Editor/ObjectMixer/ObjectMixer/Source/ObjectMixer/Private/Views/List/SObjectMixerEditorList.cpp b/Engine/Plugins/Editor/ObjectMixer/ObjectMixer/Source/ObjectMixer/Private/Views/List/SObjectMixerEditorList.cpp
--- a/Engine/Plugins/Editor/ObjectMixer/ObjectMixer/Source/ObjectMixer/Private/Views/List/SObjectMixerEditorList.cpp
+++ b/Engine/Plugins/Editor/ObjectMixer/ObjectMixer/Source/ObjectMixer/Private/Views/List/SObjectMixerEditorList.cpp
@@ -986,7 +986,7 @@ void SObjectMixerEditorList::CacheTreeState(const TArray<TWeakPtr<IObjectMixerEd
 	static LambdaType RecursivelyCacheTreeState = [](
 		const TArray<FObjectMixerEditorListRowPtr>& InObjects,
 		TArray<FTreeItemStateCache>* TreeItemStateCache,
-		TSharedPtr<STreeView<FObjectMixerEditorListRowPtr>> TreeViewPtr)
+		TSharedPtr<STreeView<FObjectMixerEditorListRowPtr>> InTreeViewPtr)
 	{
 		for (const FObjectMixerEditorListRowPtr& TreeViewItem : InObjects)
 		{
@@ -1001,14 +1001,14 @@ void SObjectMixerEditorList::CacheTreeState(const TArray<TWeakPtr<IObjectMixerEd
 					{
 						RowObject ? RowObject->GetUniqueID() : -1,
 						RowName,
-						TreeViewPtr->IsItemExpanded(TreeViewItem), 
-						TreeViewPtr->IsItemSelected(TreeViewItem),
+						InTreeViewPtr->IsItemExpanded(TreeViewItem),
+						InTreeViewPtr->IsItemSelected(TreeViewItem),
 						TreeViewItem->GetVisibilityRules()
 					}
 				);
 			}
 
-			RecursivelyCacheTreeState(TreeViewItem->GetChildRows(), TreeItemStateCache, TreeViewPtr);
+			RecursivelyCacheTreeState(TreeViewItem->GetChildRows(), TreeItemStateCache, InTreeViewPtr);
 		}
 	};
 
@@ -1042,7 +1042,7 @@ void SObjectMixerEditorList::RestoreTreeState(const TArray<TWeakPtr<IObjectMixer
 	static LambdaType RecursivelyRestoreTreeState = [](
 		const TArray<FObjectMixerEditorListRowPtr>& InObjects,
 		TArray<FTreeItemStateCache>* TreeItemStateCache,
-		TSharedPtr<STreeView<FObjectMixerEditorListRowPtr>> TreeViewPtr,
+		TSharedPtr<STreeView<FObjectMixerEditorListRowPtr>> InTreeViewPtr,
 		const bool bExpandByDefault)
 	{
 		for (const FObjectMixerEditorListRowPtr& TreeViewItem : InObjects)
@@ -1067,18 +1067,18 @@ void SObjectMixerEditorList::RestoreTreeState(const TArray<TWeakPtr<IObjectMixer
 				}
 			))
 			{
-				TreeViewPtr->SetItemExpansion(TreeViewItem, StateCachePtr->bIsExpanded);
-				TreeViewPtr->SetItemSelection(TreeViewItem, StateCachePtr->bIsSelected);
+				InTreeViewPtr->SetItemExpansion(TreeViewItem, StateCachePtr->bIsExpanded);
+				InTreeViewPtr->SetItemSelection(TreeViewItem, StateCachePtr->bIsSelected);
 				
 				TreeViewItem->SetVisibilityRules(StateCachePtr->VisibilityRules);
 			}
 			else
 			{
-				TreeViewPtr->SetItemExpansion(TreeViewItem, bExpandByDefault);
+				InTreeViewPtr->SetItemExpansion(TreeViewItem, bExpandByDefault);
 			}
 
 			RecursivelyRestoreTreeState(
-				TreeViewItem->GetChildRows(), TreeItemStateCache, TreeViewPtr, bExpandByDefault);
+				TreeViewItem->GetChildRows(), TreeItemStateCache, InTreeViewPtr, bExpandByDefault);
 		}
 	};
 	
