Object Mixer:

-Renaming Actors Inline

diff --git a/Engine/Plugins/Editor/ObjectMixer/ObjectMixer/Source/ObjectMixer/Private/Views/List/ObjectMixerEditorListRow.cpp b/Engine/Plugins/Editor/ObjectMixer/ObjectMixer/Source/ObjectMixer/Private/Views/List/ObjectMixerEditorListRow.cpp
--- a/Engine/Plugins/Editor/ObjectMixer/ObjectMixer/Source/ObjectMixer/Private/Views/List/ObjectMixerEditorListRow.cpp
+++ b/Engine/Plugins/Editor/ObjectMixer/ObjectMixer/Source/ObjectMixer/Private/Views/List/ObjectMixerEditorListRow.cpp
@@ -358,9 +358,10 @@ bool FObjectMixerEditorListRow::HasAtLeastOneChildThatIsNotSolo(const bool bRecu
 
 FText FObjectMixerEditorListRow::GetDisplayName(const bool bIsHybridRow) const
 {
-	if (!GetDisplayNameOverride().IsEmpty())
+	const FText Override = GetDisplayNameOverride();
+	if (!Override.IsEmpty())
 	{
-		return GetDisplayNameOverride();
+		return Override;
 	}
 
 	if (const UObjectMixerObjectFilter* Filter = GetObjectFilter())
diff --git a/Engine/Plugins/Editor/ObjectMixer/ObjectMixer/Source/ObjectMixer/Private/Views/List/SObjectMixerEditorList.cpp b/Engine/Plugins/Editor/ObjectMixer/ObjectMixer/Source/ObjectMixer/Private/Views/List/SObjectMixerEditorList.cpp
--- a/Engine/Plugins/Editor/ObjectMixer/ObjectMixer/Source/ObjectMixer/Private/Views/List/SObjectMixerEditorList.cpp
+++ b/Engine/Plugins/Editor/ObjectMixer/ObjectMixer/Source/ObjectMixer/Private/Views/List/SObjectMixerEditorList.cpp
@@ -1294,7 +1294,7 @@ void CreateFolderRowsForTopLevelRow(
 			{
 				FolderRow =
 					MakeShared<FObjectMixerEditorListRow>(
-						BaseActorFolder.GetPath(), FObjectMixerEditorListRow::Folder, InListView,
+						BaseActorFolder, FObjectMixerEditorListRow::Folder, InListView,
 						FText::FromName(BaseActorFolder.GetLeafName()));
 
 				FolderMap.Add(BaseActorFolder.GetPath(), FolderRow.ToSharedRef());
diff --git a/Engine/Plugins/Editor/ObjectMixer/ObjectMixer/Source/ObjectMixer/Private/Views/List/SObjectMixerEditorListRow.cpp b/Engine/Plugins/Editor/ObjectMixer/ObjectMixer/Source/ObjectMixer/Private/Views/List/SObjectMixerEditorListRow.cpp
--- a/Engine/Plugins/Editor/ObjectMixer/ObjectMixer/Source/ObjectMixer/Private/Views/List/SObjectMixerEditorListRow.cpp
+++ b/Engine/Plugins/Editor/ObjectMixer/ObjectMixer/Source/ObjectMixer/Private/Views/List/SObjectMixerEditorListRow.cpp
@@ -10,18 +10,23 @@
 
 #include "Customizations/ColorStructCustomization.h"
 #include "Editor.h"
+#include "EditorActorFolders.h"
 #include "Engine/Blueprint.h"
 #include "GameFramework/Actor.h"
 #include "Input/DragAndDrop.h"
 #include "ISinglePropertyView.h"
+#include "ObjectMixerEditorLog.h"
 #include "Modules/ModuleManager.h"
 #include "PropertyEditorModule.h"
 #include "PropertyHandle.h"
 #include "ScopedTransaction.h"
+#include "Kismet2/BlueprintEditorUtils.h"
+#include "Kismet2/ComponentEditorUtils.h"
 #include "Styling/AppStyle.h"
 #include "Styling/StyleColors.h"
 #include "Widgets/Layout/SBox.h"
 #include "Widgets/SBoxPanel.h"
+#include "Widgets/Text/SInlineEditableTextBlock.h"
 #include "Widgets/Text/STextBlock.h"
 
 #define LOCTEXT_NAMESPACE "ObjectMixerEditor"
@@ -29,6 +34,366 @@
 const FText DropFormatText = LOCTEXT("DropFormatText", "{0} {1} {2} {3}");
 const FText MultiDragFormatText = LOCTEXT("MultiDragFormatText", "{0} Items");
 
+class SInlineEditableRowNameCellWidget : public SCompoundWidget
+{
+public:
+	
+	SLATE_BEGIN_ARGS(SObjectMixerEditorListRow)
+	{}
+
+	SLATE_END_ARGS()
+	
+	void Construct(const FArguments& InArgs, TSharedRef<FObjectMixerEditorListRow> InRow, FObjectMixerEditorListRowPtr InHybridChild)
+	{
+		Item = InRow;
+		HybridChild = InHybridChild;
+		
+		TSharedRef<SHorizontalBox> HBox = SNew(SHorizontalBox);
+		
+		HBox->AddSlot()
+		.AutoWidth()
+		[
+			SNew(SImage)
+			.Image(InRow->GetObjectIconBrush())
+			.ColorAndOpacity(FSlateColor::UseForeground())
+		];
+
+		bool bNeedsStandardTextBlock = true;
+		const FText DisplayName = InRow->GetDisplayName(HybridChild.IsValid());
+		const TObjectPtr<UObject> RowObject = InRow->GetObject();
+		const bool bRowObjectIsValid = IsValid(RowObject);
+
+		if (bRowObjectIsValid)
+		{
+			if (const UClass* ActorClass = RowObject->GetClass())
+			{
+				if (UBlueprint* AsBlueprint = UBlueprint::GetBlueprintFromClass(ActorClass))
+				{					
+					bNeedsStandardTextBlock = false;
+	
+					HBox->AddSlot()
+					.Padding(FMargin(10.0, 0, 0, 0))
+					[
+						SNew(SHyperlinkWithTextHighlight)
+						.Style(FAppStyle::Get(), "Common.GotoBlueprintHyperlink")
+						.Visibility(EVisibility::Visible)
+						.Text(DisplayName)
+						.ToolTipText(LOCTEXT("ClickToEditBlueprint", "Click to edit Blueprint"))
+						.OnNavigate(this, &SInlineEditableRowNameCellWidget::OnClickBlueprintLink, AsBlueprint, RowObject.Get())
+						.HighlightText(this, &SInlineEditableRowNameCellWidget::GetHighlightText)
+					];
+				}
+			}
+		}
+
+		if (bNeedsStandardTextBlock)
+		{
+			FText TooltipText = DisplayName;
+
+			if (bRowObjectIsValid)
+			{
+				if (const UObjectMixerObjectFilter* Filter = InRow->GetObjectFilter())
+				{
+					TooltipText = Filter->GetRowTooltipText(RowObject, HybridChild.IsValid());
+				}
+			}
+			
+			HBox->AddSlot()
+			.Padding(FMargin(10.0, 0, 0, 0))
+			[
+				SNew(SInlineEditableTextBlock)
+				.Visibility(EVisibility::Visible)
+				.Justification(ETextJustify::Left)
+				.Text(DisplayName)
+				.ToolTipText(TooltipText)
+				.IsReadOnly(false)
+				.HighlightText(this, &SInlineEditableRowNameCellWidget::GetHighlightText)
+				.IsSelected_Raw(this, &SInlineEditableRowNameCellWidget::GetIsSelectedExclusively)
+				.OnTextCommitted(this, &SInlineEditableRowNameCellWidget::OnTextCommitted)
+			];
+		}
+		
+		ChildSlot
+		[
+			SNew(SBox)
+			.Visibility(EVisibility::SelfHitTestInvisible)
+			.HAlign(HAlign_Left)
+			.VAlign(VAlign_Center)
+			.Padding(FMargin(TextBlockLeftPadding, 0.f, 0.f, 0.f))
+			[
+				HBox
+			]
+		];
+	}
+
+private:
+
+	bool IsValidComponentRename(const UActorComponent* ComponentInstance, const FText& InNewText) const
+	{
+		if (!ComponentInstance)
+		{
+			return false;
+		}
+		
+		FText OutErrorMessage;
+		const FString& NewTextStr = InNewText.ToString();
+
+		if (IsValidRename(InNewText, ComponentInstance->GetName()))
+		{
+			AActor* Owner = ComponentInstance->GetOwner();
+
+			if (!Owner)
+			{
+				return false;
+			}
+
+			UBlueprint* Blueprint = nullptr;
+
+			if (const UClass* ActorClass = Owner->GetClass())
+			{
+				Blueprint = UBlueprint::GetBlueprintFromClass(ActorClass);
+				if ( Blueprint )
+				{
+					// Subobject names must conform to UObject naming conventions.
+					if (!FName::IsValidXName(NewTextStr, INVALID_OBJECTNAME_CHARACTERS, &OutErrorMessage))
+					{
+						UE_LOG(LogObjectMixerEditor, Warning, TEXT("%s"), *OutErrorMessage.ToString());
+						return false;
+					}
+			
+					AActor* ExistingNameSearchScope = ComponentInstance->GetOwner();
+			
+					if ((ExistingNameSearchScope == nullptr) && (Blueprint != nullptr))
+					{
+						ExistingNameSearchScope = Cast<AActor>(Blueprint->GeneratedClass->GetDefaultObject());
+					}
+
+					if (!FComponentEditorUtils::IsValidVariableNameString(ComponentInstance, NewTextStr))
+					{
+						OutErrorMessage = LOCTEXT("RenameFailed_EngineReservedName", "This name is reserved for engine use.");
+						UE_LOG(LogObjectMixerEditor, Warning, TEXT("%hs: %s"), __FUNCTION__, *OutErrorMessage.ToString());
+						return false;
+					}
+					else if (!FComponentEditorUtils::IsComponentNameAvailable(NewTextStr, ExistingNameSearchScope, ComponentInstance) 
+							|| !FComponentEditorUtils::IsComponentNameAvailable(NewTextStr, ComponentInstance->GetOuter(), ComponentInstance ))
+					{
+						OutErrorMessage = LOCTEXT("RenameFailed_ExistingName", "Another component already has the same name.");
+						UE_LOG(LogObjectMixerEditor, Warning, TEXT("%hs: %s"), __FUNCTION__, *OutErrorMessage.ToString());
+						return false;
+					}
+				}
+			}
+		
+			TSharedPtr<INameValidatorInterface> NameValidator;
+			if (Blueprint != nullptr)
+			{
+				NameValidator = MakeShareable(new FKismetNameValidator(Blueprint, ComponentInstance->GetFName()));
+			}
+			else
+			{
+				NameValidator = MakeShareable(new FStringSetNameValidator(ComponentInstance->GetName()));
+			}
+
+			if(NameValidator)
+			{
+				EValidatorResult ValidatorResult = NameValidator->IsValid(NewTextStr);
+				if (ValidatorResult == EValidatorResult::AlreadyInUse || ValidatorResult == EValidatorResult::LocallyInUse)
+				{
+					OutErrorMessage = FText::Format(LOCTEXT("RenameFailed_InUse", "'{0}' is in use by another variable or function!"), InNewText);
+					UE_LOG(LogObjectMixerEditor, Warning, TEXT("%hs: %s"), __FUNCTION__, *OutErrorMessage.ToString());
+				}
+				else if (ValidatorResult == EValidatorResult::Ok)
+				{
+					return true;
+				}
+			}
+		}
+
+		return false;
+	}
+
+	bool IsValidRename(const FText& NewName, const FString& OldName) const
+	{
+		FText OutErrorMessage;
+
+		if (NewName.IsEmpty())
+		{
+			OutErrorMessage = LOCTEXT("RenameFailed_LeftBlank", "Names cannot be left blank");
+			UE_LOG(LogObjectMixerEditor, Warning, TEXT("%hs: %s"), __FUNCTION__, *OutErrorMessage.ToString());
+			return false;
+		}
+
+		if (NewName.ToString().Len() >= NAME_SIZE)
+		{
+			FFormatNamedArguments Arguments;
+			Arguments.Add(TEXT("CharCount"), NAME_SIZE);
+			OutErrorMessage = FText::Format(LOCTEXT("RenameFailed_TooLong", "Names must be less than {CharCount} characters long."), Arguments);
+			UE_LOG(LogObjectMixerEditor, Warning, TEXT("%hs: %s"), __FUNCTION__, *OutErrorMessage.ToString());
+			return false;
+		}
+
+		const FString LabelString = NewName.ToString();
+		if (OldName.Equals(LabelString))
+		{
+			OutErrorMessage = LOCTEXT("RenameFailed_SameName", "Old and new names are the same.");
+			UE_LOG(LogObjectMixerEditor, Warning, TEXT("%hs: %s"), __FUNCTION__, *OutErrorMessage.ToString());
+			return false;
+		}
+
+		int32 Dummy = 0;
+		if (LabelString.FindChar('/', Dummy) || LabelString.FindChar('\\', Dummy))
+		{
+			OutErrorMessage = LOCTEXT("RenameFailed_InvalidChar", "Names cannot contain / or \\.");
+			UE_LOG(LogObjectMixerEditor, Warning, TEXT("%hs: %s"), __FUNCTION__, *OutErrorMessage.ToString());
+			return false;
+		}
+
+		return true;
+	}
+
+	static void RenameFolder(const FObjectMixerEditorListRowPtr RowPtr, const FString& TextAsString)
+	{
+		const FFolder OldFolder = RowPtr->GetFolder();
+		FName NewPath = OldFolder.GetParent().GetPath();
+		if (NewPath.IsNone())
+		{
+			NewPath = FName(*TextAsString);
+		}
+		else
+		{
+			NewPath = FName(*(NewPath.ToString() / TextAsString));
+		}
+			
+		const FFolder NewFolder(OldFolder.GetRootObject(), NewPath);
+			
+		// Transaction is built into the following method
+		FActorFolders::Get().RenameFolderInWorld(*GEditor->GetEditorWorldContext().World(), OldFolder, NewFolder);
+
+		if (const TSharedPtr<SObjectMixerEditorList> PinnedListView = RowPtr->GetListViewPtr().Pin())
+		{
+			PinnedListView->RequestRebuildList();
+		}
+	}
+
+	void RenameActor(AActor* AsActor, const FText& TrimmedLabel) const
+	{
+		if (AsActor->IsActorLabelEditable() && IsValidRename(TrimmedLabel, AsActor->GetActorLabel()))
+		{
+			const FScopedTransaction Transaction(LOCTEXT("ObjectMixerRenameActorTransaction", "Rename Actor"));
+			AsActor->Modify();
+			FActorLabelUtilities::RenameExistingActor(AsActor, TrimmedLabel.ToString());
+		}
+	}
+
+	bool RenameComponent(UActorComponent* AsComponent, const FText& TrimmedLabel) const
+	{
+		if (IsValidComponentRename(AsComponent, TrimmedLabel))
+		{
+			const FString TextAsString = TrimmedLabel.ToString();
+			
+			const ERenameFlags RenameFlags = REN_DontCreateRedirectors;
+			if(!StaticFindObject(UObject::StaticClass(), AsComponent->GetOuter(), *TextAsString))
+			{
+				const FScopedTransaction Transaction(LOCTEXT("ObjectMixerRenameComponentTransaction", "Rename Component"));
+				AsComponent->Modify();
+				AsComponent->Rename(*TextAsString, nullptr, RenameFlags);
+
+				return true;
+			}
+		}
+
+		return false;
+	}
+
+	void OnTextCommitted(const FText& InText, ETextCommit::Type CommitType) const
+	{
+		const FObjectMixerEditorListRowPtr RowPtr = Item.Pin();
+		check(RowPtr);
+		
+		const FText TrimmedLabel = FText::TrimPrecedingAndTrailing(InText);
+		if (TrimmedLabel.IsEmpty())
+		{
+			return;
+		}
+		
+		const FString TextAsString = TrimmedLabel.ToString();
+
+		if (RowPtr->GetRowType() == FObjectMixerEditorListRow::Folder &&
+			IsValidRename(TrimmedLabel, RowPtr->GetFolder().GetLeafName().ToString()))
+		{
+			RenameFolder(RowPtr, TextAsString);
+
+			return;
+		}
+		
+		UObject* RowObject = RowPtr->GetObject();
+
+		if (!RowObject)
+		{
+			return;
+		}
+
+		if (AActor* AsActor = Cast<AActor>(RowObject))
+		{
+			RenameActor(AsActor, TrimmedLabel);
+			
+			if (const TSharedPtr<SObjectMixerEditorList> PinnedListView = RowPtr->GetListViewPtr().Pin())
+			{
+				PinnedListView->RequestRebuildList();
+			}
+		}
+		else if (UActorComponent* AsComponent = Cast<UActorComponent>(RowObject))
+		{
+			if (RenameComponent(AsComponent, TrimmedLabel))
+			{
+				if (const TSharedPtr<SObjectMixerEditorList> PinnedListView = RowPtr->GetListViewPtr().Pin())
+				{
+					PinnedListView->RequestRebuildList();
+				}
+			}
+		}
+	}
+	
+	void OnClickBlueprintLink(UBlueprint* AsBlueprint, UObject* Object) const
+	{
+		if (AsBlueprint)
+		{
+			if (Object)
+			{
+				if (ensure(Object->GetClass()->ClassGeneratedBy == AsBlueprint))
+				{
+					AsBlueprint->SetObjectBeingDebugged(Object);
+				}
+			}
+			// Open the blueprint
+			GEditor->EditObject(AsBlueprint);
+		}
+	}
+
+	FText GetHighlightText() const
+	{
+		const FObjectMixerEditorListRowPtr RowPtr = HybridChild.IsValid() ? HybridChild.Pin() : Item.Pin();
+
+		if (const TSharedPtr<SObjectMixerEditorList> PinnedListView = RowPtr->GetListViewPtr().Pin())
+		{
+			return PinnedListView->GetSearchTextFromSearchInputField();
+		}
+
+		return FText::GetEmpty();
+	}
+
+	bool GetIsSelectedExclusively() const
+	{
+		return Item.Pin()->GetIsSelected() && Item.Pin()->GetSelectedTreeViewItems().Num() == 1;
+	}
+
+	TWeakPtr<FObjectMixerEditorListRow> Item;
+	TWeakPtr<FObjectMixerEditorListRow> HybridChild;
+
+	/** The offset applied to text widgets so that the text aligns with the column header text */
+	float TextBlockLeftPadding = 3.0f;
+};
+
 void SObjectMixerEditorListRow::Construct(
 	const FArguments& InArgs, const TSharedRef<STableViewBase>& InOwnerTable,
 	const TWeakPtr<FObjectMixerEditorListRow> InRow)
@@ -456,77 +821,8 @@ TSharedPtr<SWidget> SObjectMixerEditorListRow::GenerateCells(
 	
 	if (InColumnName.IsEqual(SObjectMixerEditorList::ItemNameColumnName))
 	{
-		TSharedRef<SHorizontalBox> HBox = SNew(SHorizontalBox);
-		
-		HBox->AddSlot()
-		.AutoWidth()
-		[
-			SNew(SImage)
-			.Image_Lambda([RowPtr]()
-			{
-				return RowPtr->GetObjectIconBrush();
-			})
-			.ColorAndOpacity(FSlateColor::UseForeground())
-		];
-
-		bool bNeedsStandardTextBlock = true;
-		const FText DisplayName = RowPtr->GetDisplayName(bIsHybridRow);
-		
-		if (TObjectPtr<UObject> Object = RowPtr->GetObject())
-		{
-			if (UClass* ActorClass = Object->GetClass())
-			{
-				if (UBlueprint* AsBlueprint = UBlueprint::GetBlueprintFromClass(ActorClass))
-				{					
-					bNeedsStandardTextBlock = false;
-			
-					HBox->AddSlot()
-					.Padding(FMargin(10.0, 0, 0, 0))
-					[
-						SNew(SHyperlinkWithTextHighlight)
-						.Style(FAppStyle::Get(), "Common.GotoBlueprintHyperlink")
-						.Visibility(EVisibility::Visible)
-						.Text(DisplayName)
-						.ToolTipText(LOCTEXT("ClickToEditBlueprint", "Click to edit Blueprint"))
-						.OnNavigate(this, &SObjectMixerEditorListRow::OnClickBlueprintLink, AsBlueprint, Object.Get())
-						.HighlightText(this, &SObjectMixerEditorListRow::GetHighlightText)
-					];
-				}
-			}
-		}
-
-		if (bNeedsStandardTextBlock)
-		{
-			FText TooltipText = DisplayName;
-
-			if (const UObjectMixerObjectFilter* Filter = RowPtr->GetObjectFilter())
-			{
-				if (const TObjectPtr<UObject> Object = RowPtr->GetObject())
-				{
-					TooltipText = Filter->GetRowTooltipText(Object, bIsHybridRow);
-				}
-			}
-			
-			HBox->AddSlot()
-			.Padding(FMargin(10.0, 0, 0, 0))
-			[
-				SNew(STextBlock)
-				.Visibility(EVisibility::SelfHitTestInvisible)
-				.Justification(ETextJustify::Left)
-				.Text(DisplayName)
-				.ToolTipText(TooltipText)
-				.HighlightText(this, &SObjectMixerEditorListRow::GetHighlightText)
-			];
-		}
-		
-		return SNew(SBox)
-				.Visibility(EVisibility::SelfHitTestInvisible)
-				.HAlign(HAlign_Left)
-				.VAlign(VAlign_Center)
-				.Padding(FMargin(TextBlockLeftPadding, 0.f, 0.f, 0.f))
-				[
-					HBox
-				];
+		return SNew(SInlineEditableRowNameCellWidget,
+			Item.Pin().ToSharedRef(), HybridRowIndex != INDEX_NONE ? GetHybridChildOrRowItemIfNull() : nullptr);
 	}
 
 	if (InColumnName.IsEqual(SObjectMixerEditorList::EditorVisibilityColumnName))
@@ -692,33 +988,4 @@ void SObjectMixerEditorListRow::OnPropertyChanged(const FName PropertyName) cons
 	}
 }
 
-void SObjectMixerEditorListRow::OnClickBlueprintLink(UBlueprint* AsBlueprint, UObject* Object)
-{
-	if (AsBlueprint)
-	{
-		if (Object)
-		{
-			if (ensure(Object->GetClass()->ClassGeneratedBy == AsBlueprint))
-			{
-				AsBlueprint->SetObjectBeingDebugged(Object);
-			}
-		}
-		// Open the blueprint
-		GEditor->EditObject(AsBlueprint);
-	}
-}
-
-FText SObjectMixerEditorListRow::GetHighlightText() const
-{
-	check (Item.IsValid());
-	const FObjectMixerEditorListRowPtr RowPtr = GetHybridChildOrRowItemIfNull();
-
-	if (const TSharedPtr<SObjectMixerEditorList> PinnedListView = RowPtr->GetListViewPtr().Pin())
-	{
-		return PinnedListView->GetSearchTextFromSearchInputField();
-	}
-
-	return FText::GetEmpty();
-}
-
 #undef LOCTEXT_NAMESPACE
diff --git a/Engine/Plugins/Editor/ObjectMixer/ObjectMixer/Source/ObjectMixer/Public/Views/List/ObjectMixerEditorListRow.h b/Engine/Plugins/Editor/ObjectMixer/ObjectMixer/Source/ObjectMixer/Public/Views/List/ObjectMixerEditorListRow.h
--- a/Engine/Plugins/Editor/ObjectMixer/ObjectMixer/Source/ObjectMixer/Public/Views/List/ObjectMixerEditorListRow.h
+++ b/Engine/Plugins/Editor/ObjectMixer/ObjectMixer/Source/ObjectMixer/Public/Views/List/ObjectMixerEditorListRow.h
@@ -2,11 +2,13 @@
 
 #pragma once
 
+#include "ObjectFilter/ObjectMixerEditorObjectFilter.h"
+
 #include "Containers/Array.h"
 #include "Containers/UnrealString.h"
 #include "DragAndDrop/DecoratedDragDropOp.h"
+#include "Folder.h"
 #include "Layout/Visibility.h"
-#include "ObjectFilter/ObjectMixerEditorObjectFilter.h"
 #include "PropertyHandle.h"
 #include "Templates/SharedPointer.h"
 #include "UObject/Object.h"
@@ -89,9 +91,9 @@ struct OBJECTMIXEREDITOR_API FObjectMixerEditorListRow final : TSharedFromThis<F
 	{}
 
 	FObjectMixerEditorListRow(
-		const FName InFolderPath, const EObjectMixerEditorListRowType InRowType, 
+		const FFolder InFolder, const EObjectMixerEditorListRowType InRowType, 
 		const TSharedRef<SObjectMixerEditorList>& InListView, const FText& InDisplayNameOverride = FText::GetEmpty())
-	: FolderPath(InFolderPath)
+	: FolderRef(InFolder)
 	, RowType(InRowType)
 	, ListViewPtr(InListView)
 	, DisplayNameOverride(InDisplayNameOverride)
@@ -107,9 +109,14 @@ struct OBJECTMIXEREDITOR_API FObjectMixerEditorListRow final : TSharedFromThis<F
 		return nullptr;
 	}
 
+	[[nodiscard]] const FFolder& GetFolder() const
+	{
+		return FolderRef;
+	}
+
 	[[nodiscard]] FName GetFolderPath() const
 	{
-		return FolderPath;
+		return FolderRef.GetPath();
 	}
 
 	UObjectMixerObjectFilter* GetObjectFilter() const;
@@ -214,7 +221,7 @@ private:
 	FObjectMixerEditorListRowPtr GetAsShared();
 
 	TWeakObjectPtr<UObject> ObjectRef;
-	FName FolderPath = NAME_None;
+	FFolder FolderRef;
 	EObjectMixerEditorListRowType RowType = MatchingObject;
 	TArray<FObjectMixerEditorListRowPtr> ChildRows;
 	
diff --git a/Engine/Plugins/Editor/ObjectMixer/ObjectMixer/Source/ObjectMixer/Public/Views/List/SObjectMixerEditorListRow.h b/Engine/Plugins/Editor/ObjectMixer/ObjectMixer/Source/ObjectMixer/Public/Views/List/SObjectMixerEditorListRow.h
--- a/Engine/Plugins/Editor/ObjectMixer/ObjectMixer/Source/ObjectMixer/Public/Views/List/SObjectMixerEditorListRow.h
+++ b/Engine/Plugins/Editor/ObjectMixer/ObjectMixer/Source/ObjectMixer/Public/Views/List/SObjectMixerEditorListRow.h
@@ -71,9 +71,6 @@ private:
 	const FSlateBrush* SoloOnBrush = nullptr;
 	const FSlateBrush* SoloOffHoveredBrush = nullptr;
 
-	/** The offset applied to text widgets so that the text aligns with the column header text */
-	float TextBlockLeftPadding = 3.0f;
-
 	/** Hybrid Rows are a combination of an actor and a single child subobject */
 	int32 HybridRowIndex = INDEX_NONE;
 
