Object Mixer:

-Reimplement Edit context menu
-Fix: Edit Hotkeys Not Working when panel is undocked

diff --git a/Engine/Plugins/Editor/ObjectMixer/ObjectMixer/Source/ObjectMixer/Private/ObjectMixerEditorModule.cpp b/Engine/Plugins/Editor/ObjectMixer/ObjectMixer/Source/ObjectMixer/Private/ObjectMixerEditorModule.cpp
--- a/Engine/Plugins/Editor/ObjectMixer/ObjectMixer/Source/ObjectMixer/Private/ObjectMixerEditorModule.cpp
+++ b/Engine/Plugins/Editor/ObjectMixer/ObjectMixer/Source/ObjectMixer/Private/ObjectMixerEditorModule.cpp
@@ -18,6 +18,7 @@
 #include "ToolMenus.h"
 #include "WorkspaceMenuStructure.h"
 #include "WorkspaceMenuStructureModule.h"
+#include "Views/Widgets/ObjectMixerEditorListMenuContext.h"
 
 const FName FObjectMixerEditorModule::BaseObjectMixerModuleName("ObjectMixerEditor");
 
@@ -142,6 +143,14 @@ void FObjectMixerEditorModule::RefreshList() const
 	}
 }
 
+void FObjectMixerEditorModule::OnRenameCommand()
+{
+	if (MainPanel.IsValid())
+	{
+		MainPanel->OnRenameCommand();
+	}
+}
+
 void FObjectMixerEditorModule::RequestSyncEditorSelectionToListSelection()
 {
 	if (MainPanel.IsValid())
diff --git a/Engine/Plugins/Editor/ObjectMixer/ObjectMixer/Source/ObjectMixer/Private/Views/List/ObjectMixerEditorList.cpp b/Engine/Plugins/Editor/ObjectMixer/ObjectMixer/Source/ObjectMixer/Private/Views/List/ObjectMixerEditorList.cpp
--- a/Engine/Plugins/Editor/ObjectMixer/ObjectMixer/Source/ObjectMixer/Private/Views/List/ObjectMixerEditorList.cpp
+++ b/Engine/Plugins/Editor/ObjectMixer/ObjectMixer/Source/ObjectMixer/Private/Views/List/ObjectMixerEditorList.cpp
@@ -87,6 +87,14 @@ void FObjectMixerEditorList::RefreshList() const
 	}
 }
 
+void FObjectMixerEditorList::OnRenameCommand()
+{
+	if (ListWidget.IsValid())
+	{
+		ListWidget->OnRenameCommand();;
+	}
+}
+
 void FObjectMixerEditorList::RequestSyncEditorSelectionToListSelection() const
 {
 	if (ListWidget.IsValid())
diff --git a/Engine/Plugins/Editor/ObjectMixer/ObjectMixer/Source/ObjectMixer/Private/Views/List/SObjectMixerEditorList.cpp b/Engine/Plugins/Editor/ObjectMixer/ObjectMixer/Source/ObjectMixer/Private/Views/List/SObjectMixerEditorList.cpp
--- a/Engine/Plugins/Editor/ObjectMixer/ObjectMixer/Source/ObjectMixer/Private/Views/List/SObjectMixerEditorList.cpp
+++ b/Engine/Plugins/Editor/ObjectMixer/ObjectMixer/Source/ObjectMixer/Private/Views/List/SObjectMixerEditorList.cpp
@@ -200,6 +200,14 @@ void SObjectMixerEditorList::RequestRebuildList(const FString& InItemToScrollTo)
 	bIsRebuildRequested = true;
 }
 
+void SObjectMixerEditorList::OnRenameCommand()
+{
+	if (GetSelectedTreeViewItemCount() == 1)
+	{
+		GetSelectedTreeViewItems()[0]->CallOnRenameCommandDelegate();
+	}
+}
+
 TArray<FObjectMixerEditorListRowPtr> SObjectMixerEditorList::GetSelectedTreeViewItems() const
 {
 	return TreeViewPtr->GetSelectedItems();
diff --git a/Engine/Plugins/Editor/ObjectMixer/ObjectMixer/Source/ObjectMixer/Private/Views/List/SObjectMixerEditorListRow.cpp b/Engine/Plugins/Editor/ObjectMixer/ObjectMixer/Source/ObjectMixer/Private/Views/List/SObjectMixerEditorListRow.cpp
--- a/Engine/Plugins/Editor/ObjectMixer/ObjectMixer/Source/ObjectMixer/Private/Views/List/SObjectMixerEditorListRow.cpp
+++ b/Engine/Plugins/Editor/ObjectMixer/ObjectMixer/Source/ObjectMixer/Private/Views/List/SObjectMixerEditorListRow.cpp
@@ -47,6 +47,8 @@ public:
 	{
 		Item = InRow;
 		HybridChild = InHybridChild;
+
+		InRow->OnRenameCommand().BindRaw(this, &SInlineEditableRowNameCellWidget::EnterEditingMode);
 		
 		TSharedRef<SHorizontalBox> HBox = SNew(SHorizontalBox);
 		
@@ -101,7 +103,7 @@ public:
 			HBox->AddSlot()
 			.Padding(FMargin(10.0, 0, 0, 0))
 			[
-				SNew(SInlineEditableTextBlock)
+				SAssignNew(EditableTextBlock, SInlineEditableTextBlock)
 				.Visibility(EVisibility::Visible)
 				.Justification(ETextJustify::Left)
 				.Text(DisplayName)
@@ -126,6 +128,24 @@ public:
 		];
 	}
 
+	virtual ~SInlineEditableRowNameCellWidget() override
+	{
+		if (const FObjectMixerEditorListRowPtr ItemPin = Item.Pin())
+		{
+			ItemPin->OnRenameCommand().Unbind();
+		}
+		
+		Item.Reset();
+		HybridChild.Reset();
+
+		EditableTextBlock.Reset();
+	}
+
+	void EnterEditingMode() const
+	{
+		EditableTextBlock->EnterEditingMode();
+	}
+
 private:
 
 	bool IsValidComponentRename(const UActorComponent* ComponentInstance, const FText& InNewText) const
@@ -390,6 +410,8 @@ private:
 	TWeakPtr<FObjectMixerEditorListRow> Item;
 	TWeakPtr<FObjectMixerEditorListRow> HybridChild;
 
+	TSharedPtr<SInlineEditableTextBlock> EditableTextBlock;
+
 	/** The offset applied to text widgets so that the text aligns with the column header text */
 	float TextBlockLeftPadding = 3.0f;
 };
diff --git a/Engine/Plugins/Editor/ObjectMixer/ObjectMixer/Source/ObjectMixer/Private/Views/MainPanel/ObjectMixerEditorMainPanel.cpp b/Engine/Plugins/Editor/ObjectMixer/ObjectMixer/Source/ObjectMixer/Private/Views/MainPanel/ObjectMixerEditorMainPanel.cpp
--- a/Engine/Plugins/Editor/ObjectMixer/ObjectMixer/Source/ObjectMixer/Private/Views/MainPanel/ObjectMixerEditorMainPanel.cpp
+++ b/Engine/Plugins/Editor/ObjectMixer/ObjectMixer/Source/ObjectMixer/Private/Views/MainPanel/ObjectMixerEditorMainPanel.cpp
@@ -2,13 +2,45 @@
 
 #include "Views/MainPanel/ObjectMixerEditorMainPanel.h"
 
+#include "LevelEditorActions.h"
 #include "ObjectMixerEditorSerializedData.h"
+#include "Framework/Commands/GenericCommands.h"
 #include "Views/List/ObjectMixerEditorList.h"
 #include "Views/MainPanel/SObjectMixerEditorMainPanel.h"
 
 void FObjectMixerEditorMainPanel::Init()
 {
 	RegenerateListModel();
+	RegisterAndMapContextMenuCommands();
+}
+
+void FObjectMixerEditorMainPanel::RegisterAndMapContextMenuCommands()
+{
+	ObjectMixerCommands = MakeShared<FUICommandList>();
+
+	ObjectMixerCommands->MapAction(FGenericCommands::Get().Cut,
+		FExecuteAction::CreateStatic( &FLevelEditorActionCallbacks::ExecuteExecCommand, FString( TEXT("EDIT CUT") ) ),
+		FCanExecuteAction::CreateStatic( &FLevelEditorActionCallbacks::Cut_CanExecute )
+	);
+	ObjectMixerCommands->MapAction(FGenericCommands::Get().Copy,
+		FExecuteAction::CreateStatic( &FLevelEditorActionCallbacks::ExecuteExecCommand, FString( TEXT("EDIT COPY") ) ),
+		FCanExecuteAction::CreateStatic( &FLevelEditorActionCallbacks::Copy_CanExecute )
+	);
+	ObjectMixerCommands->MapAction(FGenericCommands::Get().Paste,
+		FExecuteAction::CreateStatic( &FLevelEditorActionCallbacks::ExecuteExecCommand, FString( TEXT("EDIT PASTE") ) ),
+		FCanExecuteAction::CreateStatic( &FLevelEditorActionCallbacks::Paste_CanExecute )
+	);
+	ObjectMixerCommands->MapAction(FGenericCommands::Get().Duplicate,
+		FExecuteAction::CreateStatic( &FLevelEditorActionCallbacks::ExecuteExecCommand, FString( TEXT("DUPLICATE") ) ),
+		FCanExecuteAction::CreateStatic( &FLevelEditorActionCallbacks::Duplicate_CanExecute )
+	);
+	ObjectMixerCommands->MapAction(FGenericCommands::Get().Delete,
+		FExecuteAction::CreateStatic( &FLevelEditorActionCallbacks::ExecuteExecCommand, FString( TEXT("DELETE") ) ),
+		FCanExecuteAction::CreateStatic( &FLevelEditorActionCallbacks::Delete_CanExecute )
+	);
+	ObjectMixerCommands->MapAction(FGenericCommands::Get().Rename,
+		FUIAction(FExecuteAction::CreateRaw(this, &FObjectMixerEditorMainPanel::OnRenameCommand))
+	);
 }
 
 TSharedRef<SWidget> FObjectMixerEditorMainPanel::GetOrCreateWidget()
@@ -44,6 +76,14 @@ void FObjectMixerEditorMainPanel::RefreshList() const
 	}
 }
 
+void FObjectMixerEditorMainPanel::OnRenameCommand()
+{
+	if (EditorListModel.IsValid())
+	{
+		EditorListModel->OnRenameCommand();
+	}
+}
+
 void FObjectMixerEditorMainPanel::RequestSyncEditorSelectionToListSelection() const
 {
 	if (EditorListModel.IsValid())
diff --git a/Engine/Plugins/Editor/ObjectMixer/ObjectMixer/Source/ObjectMixer/Private/Views/MainPanel/SObjectMixerEditorMainPanel.cpp b/Engine/Plugins/Editor/ObjectMixer/ObjectMixer/Source/ObjectMixer/Private/Views/MainPanel/SObjectMixerEditorMainPanel.cpp
--- a/Engine/Plugins/Editor/ObjectMixer/ObjectMixer/Source/ObjectMixer/Private/Views/MainPanel/SObjectMixerEditorMainPanel.cpp
+++ b/Engine/Plugins/Editor/ObjectMixer/ObjectMixer/Source/ObjectMixer/Private/Views/MainPanel/SObjectMixerEditorMainPanel.cpp
@@ -80,6 +80,21 @@ void SObjectMixerEditorMainPanel::Construct(
 	SetSingleCollectionSelection();
 }
 
+FReply SObjectMixerEditorMainPanel::OnKeyDown(const FGeometry& MyGeometry, const FKeyEvent& InKeyEvent)
+{
+	// Check to see if any actions can be processed
+	// If we are in debug mode do not process commands
+	if (FSlateApplication::Get().IsNormalExecution())
+	{
+		if (GetMainPanelModel().Pin()->ObjectMixerCommands->ProcessCommandBindings(InKeyEvent))
+		{
+			return FReply::Handled();
+		}
+	}
+	
+	return SCompoundWidget::OnKeyDown(MyGeometry, InKeyEvent);
+}
+
 TSharedRef<SWidget> SObjectMixerEditorMainPanel::GenerateToolbar()
 {
 	TSharedRef<SHorizontalBox> ToolbarBox = SNew(SHorizontalBox);
diff --git a/Engine/Plugins/Editor/ObjectMixer/ObjectMixer/Source/ObjectMixer/Private/Views/Widgets/ObjectMixerEditorListMenuContext.cpp b/Engine/Plugins/Editor/ObjectMixer/ObjectMixer/Source/ObjectMixer/Private/Views/Widgets/ObjectMixerEditorListMenuContext.cpp
--- a/Engine/Plugins/Editor/ObjectMixer/ObjectMixer/Source/ObjectMixer/Private/Views/Widgets/ObjectMixerEditorListMenuContext.cpp
+++ b/Engine/Plugins/Editor/ObjectMixer/ObjectMixer/Source/ObjectMixer/Private/Views/Widgets/ObjectMixerEditorListMenuContext.cpp
@@ -8,9 +8,11 @@
 #include "Elements/Interfaces/TypedElementObjectInterface.h"
 #include "Framework/MultiBox/MultiBoxDefs.h"
 #include "LevelEditor.h"
+#include "LevelEditorActions.h"
 #include "LevelEditorContextMenu.h"
 #include "LevelEditorMenuContext.h"
 #include "LevelEditorViewport.h"
+#include "ObjectMixerEditorModule.h"
 #include "Selection.h"
 #include "ToolMenu.h"
 #include "ToolMenus.h"
@@ -18,6 +20,7 @@
 #include "ToolMenuDelegates.h"
 #include "ToolMenuEntry.h"
 #include "ToolMenuSection.h"
+#include "Framework/Commands/GenericCommands.h"
 #include "Widgets/Layout/SBox.h"
 #include "Widgets/Input/SEditableTextBox.h"
 
@@ -104,91 +107,94 @@ TSharedPtr<SWidget> UObjectMixerEditorListMenuContext::BuildContextMenu(const FO
 
 void UObjectMixerEditorListMenuContext::RegisterObjectMixerContextMenu()
 {
-	UToolMenus* ToolMenus = UToolMenus::Get();
-	
-	if (ToolMenus->IsMenuRegistered("LevelEditor.ActorContextMenu"))
-	{
-		UToolMenu* Menu = ToolMenus->ExtendMenu("LevelEditor.ActorContextMenu");
-	
-		Menu->AddDynamicSection("DynamicCollectionsSection", FNewToolMenuDelegate::CreateLambda(
+	if (UToolMenu* ActorContextMenu = UToolMenus::Get()->ExtendMenu("LevelEditor.ActorContextMenu"))
+	{		
+		ActorContextMenu->AddDynamicSection("DynamicCollectionsSection", FNewToolMenuDelegate::CreateLambda(
 			[](UToolMenu* InMenu)
 			{
+				// Ensure proper context
 				UObjectMixerEditorListMenuContext* Context = InMenu->FindContext<UObjectMixerEditorListMenuContext>();
 				if (!Context || Context->Data.SelectedItems.Num() == 0)
 				{
 					return;
 				}
-	
-				FToolMenuSection& Section = InMenu->FindOrAddSection("ObjectMixerCollections");
-				Section.Label = LOCTEXT("ObjectMixerCollectionsSectionName", "Mixer Collections");
-				
-				FObjectMixerEditorListMenuContextData& ContextData = Context->Data;
-			
-				Section.AddSubMenu(
-					"SelectCollectionsSubMenu",
-					LOCTEXT("SelectCollectionsSubmenu", "Select or Add Collection"),
-					LOCTEXT("SelectCollectionsSubmenu_Tooltip", "Select the collection to which you wish to assign this object."),
-					FNewToolMenuDelegate::CreateLambda(
-						[ContextData](UToolMenu* Menu)
-						{
-							FToolMenuEntry Args;
-							Args.Type = EMultiBlockType::Widget;
-							Args.MakeCustomWidget.BindLambda(
-								[ContextData](const FToolMenuContext&, const FToolMenuCustomWidgetContext&)
-								{
-									return SNew(SBox)
-											.MinDesiredWidth(200)
-											.Padding(8, 0)
-											[
-												SNew(SEditableTextBox)
-												.HintText(LOCTEXT("NewCollectionEditableHintText", "Enter a new collection name..."))
-												.OnTextCommitted_Static(&UObjectMixerEditorListMenuContext::OnTextCommitted, ContextData)
-											]
-									;
-								}
-							);
-							Menu->AddMenuEntry("NewCollectionInput", Args);
-							
-							if (const TSharedPtr<FObjectMixerEditorMainPanel> MainPanel = ContextData.MainPanelPtr.Pin())
-							{
-								if (TArray<FName> Collections = MainPanel->GetAllCollectionNames(); Collections.Num() > 0)
-								{
-									FToolMenuSection& Section = Menu->FindOrAddSection("Collections");
-								   Section.Label = LOCTEXT("CollectionsSectionName", "Collections");
-
-								   Collections.StableSort([](const FName& A, const FName B)
-								   {
-									   return A.LexicalLess(B);
-								   });
-								
-								   for (const FName& Key : Collections)
-								   {
-									   const FText KeyText = FText::FromName(Key);
-									
-									   Section.AddMenuEntry(
-										   Key,
-										   KeyText,
-										   FText::Format(LOCTEXT("AddObjectsToCollectionTooltipFormat", "Add selected to collection '{0}'"), KeyText),
-										   FSlateIcon(),
-										   FUIAction(
-											   FExecuteAction::CreateStatic(&UObjectMixerEditorListMenuContext::OnClickCollectionMenuEntry, Key, ContextData),
-											   FCanExecuteAction(),
-											   FIsActionChecked::CreateStatic(&UObjectMixerEditorListMenuContext::AreAllObjectsInCollection, Key, ContextData)
-										   ),
-										   EUserInterfaceActionType::ToggleButton
-									   );
-								   }
-								}
-							}
-						}
-					)
-				);
+
+				const FObjectMixerEditorListMenuContextData& ContextData = Context->Data;
+
+				ReplaceEditSubMenu(ContextData);
+
+				// Collections submenu
+				{
+					FToolMenuSection& Section = InMenu->FindOrAddSection("ObjectMixerCollections");
+					Section.Label = LOCTEXT("ObjectMixerCollectionsSectionName", "Mixer Collections");
+		
+					Section.AddSubMenu(
+						"SelectCollectionsSubMenu",
+						LOCTEXT("SelectCollectionsSubmenu", "Select or Add Collection"),
+						LOCTEXT("SelectCollectionsSubmenu_Tooltip", "Select the collection to which you wish to assign this object."),
+						FNewToolMenuDelegate::CreateStatic(&UObjectMixerEditorListMenuContext::CreateSelectCollectionsSubMenu, ContextData)
+					);
+				}
 			}),
 			FToolMenuInsert(NAME_None,EToolMenuInsertType::First)
 		);
 	}
 }
 
+void UObjectMixerEditorListMenuContext::CreateSelectCollectionsSubMenu(UToolMenu* Menu, FObjectMixerEditorListMenuContextData ContextData)
+{
+	FToolMenuEntry Args;
+	Args.Type = EMultiBlockType::Widget;
+	Args.MakeCustomWidget.BindLambda(
+		[ContextData](const FToolMenuContext&, const FToolMenuCustomWidgetContext&)
+		{
+			return
+				SNew(SBox)
+				.MinDesiredWidth(200)
+				.Padding(8, 0)
+				[
+					SNew(SEditableTextBox)
+					.HintText(LOCTEXT("NewCollectionEditableHintText", "Enter a new collection name..."))
+					.OnTextCommitted_Static(&UObjectMixerEditorListMenuContext::OnTextCommitted, ContextData)
+				]
+			;
+		}
+	);
+	Menu->AddMenuEntry("NewCollectionInput", Args);
+						
+	if (const TSharedPtr<FObjectMixerEditorMainPanel> MainPanel = ContextData.MainPanelPtr.Pin())
+	{
+		if (TArray<FName> Collections = MainPanel->GetAllCollectionNames(); Collections.Num() > 0)
+		{
+			FToolMenuSection& Section = Menu->FindOrAddSection("Collections");
+			Section.Label = LOCTEXT("CollectionsSectionName", "Collections");
+
+			Collections.StableSort([](const FName& A, const FName B)
+			{
+				return A.LexicalLess(B);
+			});
+							
+			for (const FName& Key : Collections)
+			{
+				const FText KeyText = FText::FromName(Key);
+								
+				Section.AddMenuEntry(
+					Key,
+					KeyText,
+					FText::Format(LOCTEXT("AddObjectsToCollectionTooltipFormat", "Add selected to collection '{0}'"), KeyText),
+					FSlateIcon(),
+					FUIAction(
+						FExecuteAction::CreateStatic(&UObjectMixerEditorListMenuContext::OnClickCollectionMenuEntry, Key, ContextData),
+						FCanExecuteAction(),
+						FIsActionChecked::CreateStatic(&UObjectMixerEditorListMenuContext::AreAllObjectsInCollection, Key, ContextData)
+					),
+					EUserInterfaceActionType::ToggleButton
+				);
+			}
+		}
+	}
+}
+
 void UObjectMixerEditorListMenuContext::OnTextCommitted(const FText& InText, ETextCommit::Type InCommitType, const FObjectMixerEditorListMenuContextData ContextData)
 {
 	if (InCommitType == ETextCommit::OnEnter)
@@ -268,4 +274,54 @@ bool UObjectMixerEditorListMenuContext::AreAllObjectsInCollection(const FName Ke
 	return bAreAllSelectedObjectsInCollection;
 }
 
+void UObjectMixerEditorListMenuContext::ReplaceEditSubMenu(const FObjectMixerEditorListMenuContextData& ContextData)
+{
+	if (UToolMenu* EditSubMenu = UToolMenus::Get()->ExtendMenu("LevelEditor.ActorContextMenu.EditSubMenu"))
+	{
+		// Remove existing edit sub menu options
+		FCustomizedToolMenu* MenuCustomization = EditSubMenu->AddMenuCustomization();
+		{
+			MenuCustomization->AddEntry(TEXT("Cut"))->Visibility = ECustomizedToolMenuVisibility::Hidden;
+			MenuCustomization->AddEntry(TEXT("Copy"))->Visibility = ECustomizedToolMenuVisibility::Hidden;
+			MenuCustomization->AddEntry(TEXT("Paste"))->Visibility = ECustomizedToolMenuVisibility::Hidden;
+			MenuCustomization->AddEntry(TEXT("Duplicate"))->Visibility = ECustomizedToolMenuVisibility::Hidden;
+			MenuCustomization->AddEntry(TEXT("Delete"))->Visibility = ECustomizedToolMenuVisibility::Hidden;
+			MenuCustomization->AddEntry(TEXT("Rename"))->Visibility = ECustomizedToolMenuVisibility::Hidden;
+		}
+
+		// Add our menu widget to the edit submenu
+		// It must be a widget entry because we hid the entries above, so they won't show as regular menu entries since they share names
+		FToolMenuEntry Entry;
+		Entry.Name = TEXT("ObjectMixerGenericCommands");
+		Entry.Type = EMultiBlockType::Widget;
+		FToolMenuEntryWidgetData WidgetData;
+		WidgetData.bNoIndent = true;
+		WidgetData.bNoPadding = true;
+		Entry.WidgetData = WidgetData;
+		Entry.MakeCustomWidget.BindLambda(
+			[ContextData](const FToolMenuContext&, const FToolMenuCustomWidgetContext&)
+			{
+				if (TSharedPtr<FObjectMixerEditorMainPanel> PinnedMainPanel = ContextData.MainPanelPtr.Pin())
+				{
+					// Add options with our mapped command list to an FMenuBuilder
+					FMenuBuilder Builder(true, PinnedMainPanel->ObjectMixerCommands);
+					{
+						Builder.AddMenuEntry( FGenericCommands::Get().Cut);
+						Builder.AddMenuEntry( FGenericCommands::Get().Copy);
+						Builder.AddMenuEntry( FGenericCommands::Get().Paste);
+						Builder.AddMenuEntry( FGenericCommands::Get().Duplicate);
+						Builder.AddMenuEntry( FGenericCommands::Get().Delete);
+						Builder.AddMenuEntry( FGenericCommands::Get().Rename);
+					}
+
+					return Builder.MakeWidget();
+				}
+
+				return SNullWidget::NullWidget;
+			}
+		);
+		EditSubMenu->AddMenuEntry(NAME_None, Entry);
+	}
+}
+
 #undef LOCTEXT_NAMESPACE
diff --git a/Engine/Plugins/Editor/ObjectMixer/ObjectMixer/Source/ObjectMixer/Public/ObjectMixerEditorModule.h b/Engine/Plugins/Editor/ObjectMixer/ObjectMixer/Source/ObjectMixer/Public/ObjectMixerEditorModule.h
--- a/Engine/Plugins/Editor/ObjectMixer/ObjectMixer/Source/ObjectMixer/Public/ObjectMixerEditorModule.h
+++ b/Engine/Plugins/Editor/ObjectMixer/ObjectMixer/Source/ObjectMixer/Public/ObjectMixerEditorModule.h
@@ -38,8 +38,11 @@ public:
 	 * Useful for when the list state has gone stale but the variable count has not changed.
 	 */
 	virtual void RefreshList() const;
+
+	/** Called when the Rename command is executed from the UI or hotkey. */
+	virtual void OnRenameCommand();
 	
-	void RequestSyncEditorSelectionToListSelection();
+	virtual void RequestSyncEditorSelectionToListSelection();
 
 	void RegisterMenuGroup();
 	void UnregisterMenuGroup();
diff --git a/Engine/Plugins/Editor/ObjectMixer/ObjectMixer/Source/ObjectMixer/Public/Views/List/ObjectMixerEditorList.h b/Engine/Plugins/Editor/ObjectMixer/ObjectMixer/Source/ObjectMixer/Public/Views/List/ObjectMixerEditorList.h
--- a/Engine/Plugins/Editor/ObjectMixer/ObjectMixer/Source/ObjectMixer/Public/Views/List/ObjectMixerEditorList.h
+++ b/Engine/Plugins/Editor/ObjectMixer/ObjectMixer/Source/ObjectMixer/Public/Views/List/ObjectMixerEditorList.h
@@ -37,6 +37,9 @@ public:
 	 */
 	void RefreshList() const;
 
+	/** Called when the Rename command is executed from the UI or hotkey. */
+	void OnRenameCommand();
+
 	void RequestSyncEditorSelectionToListSelection() const;
 
 	void ExecuteListViewSearchOnAllRows(const FString& SearchString, const bool bShouldRefreshAfterward = true);
diff --git a/Engine/Plugins/Editor/ObjectMixer/ObjectMixer/Source/ObjectMixer/Public/Views/List/ObjectMixerEditorListRow.h b/Engine/Plugins/Editor/ObjectMixer/ObjectMixer/Source/ObjectMixer/Public/Views/List/ObjectMixerEditorListRow.h
--- a/Engine/Plugins/Editor/ObjectMixer/ObjectMixer/Source/ObjectMixer/Public/Views/List/ObjectMixerEditorListRow.h
+++ b/Engine/Plugins/Editor/ObjectMixer/ObjectMixer/Source/ObjectMixer/Public/Views/List/ObjectMixerEditorListRow.h
@@ -12,7 +12,6 @@
 #include "PropertyHandle.h"
 #include "Templates/SharedPointer.h"
 #include "UObject/Object.h"
-#include "UObject/WeakObjectPtr.h"
 
 class SObjectMixerEditorList;
 class UObjectMixerObjectFilter;
@@ -22,6 +21,8 @@ struct FSlateBrush;
 struct FObjectMixerEditorListRow;
 typedef TSharedPtr<FObjectMixerEditorListRow> FObjectMixerEditorListRowPtr;
 
+DECLARE_DELEGATE(FOnRenameCommand)
+
 class FObjectMixerListRowDragDropOp : public FDecoratedDragDropOp
 {
 public:
@@ -213,6 +214,16 @@ struct OBJECTMIXEREDITOR_API FObjectMixerEditorListRow final : TSharedFromThis<F
 	void SetRowSoloState(const bool bNewSolo);
 
 	void ClearSoloRows() const;
+
+	FOnRenameCommand& OnRenameCommand()
+	{
+		return OnRenameCommandDelegate;
+	}
+
+	void CallOnRenameCommandDelegate()
+	{
+		OnRenameCommandDelegate.Execute();
+	}
 	
 	TMap<FName, TWeakPtr<IPropertyHandle>> PropertyNamesToHandles;
 
@@ -245,4 +256,6 @@ private:
 	bool bShouldExpandAllChildren = false;
 
 	int32 CachedHybridRowIndex = INDEX_NONE;
+
+	FOnRenameCommand OnRenameCommandDelegate;
 };
diff --git a/Engine/Plugins/Editor/ObjectMixer/ObjectMixer/Source/ObjectMixer/Public/Views/List/SObjectMixerEditorList.h b/Engine/Plugins/Editor/ObjectMixer/ObjectMixer/Source/ObjectMixer/Public/Views/List/SObjectMixerEditorList.h
--- a/Engine/Plugins/Editor/ObjectMixer/ObjectMixer/Source/ObjectMixer/Public/Views/List/SObjectMixerEditorList.h
+++ b/Engine/Plugins/Editor/ObjectMixer/ObjectMixer/Source/ObjectMixer/Public/Views/List/SObjectMixerEditorList.h
@@ -98,6 +98,9 @@ public:
 	 */
 	void RequestRebuildList(const FString& InItemToScrollTo = "");
 
+	/** Called when the Rename command is executed from the UI or hotkey. */
+	void OnRenameCommand();
+
 	[[nodiscard]] TArray<FObjectMixerEditorListRowPtr> GetSelectedTreeViewItems() const;
 	int32 GetSelectedTreeViewItemCount() const;
 	
diff --git a/Engine/Plugins/Editor/ObjectMixer/ObjectMixer/Source/ObjectMixer/Public/Views/MainPanel/ObjectMixerEditorMainPanel.h b/Engine/Plugins/Editor/ObjectMixer/ObjectMixer/Source/ObjectMixer/Public/Views/MainPanel/ObjectMixerEditorMainPanel.h
--- a/Engine/Plugins/Editor/ObjectMixer/ObjectMixer/Source/ObjectMixer/Public/Views/MainPanel/ObjectMixerEditorMainPanel.h
+++ b/Engine/Plugins/Editor/ObjectMixer/ObjectMixer/Source/ObjectMixer/Public/Views/MainPanel/ObjectMixerEditorMainPanel.h
@@ -7,6 +7,7 @@
 #include "UObject/StrongObjectPtr.h"
 #include "Widgets/SWidget.h"
 
+class FUICommandList;
 class FObjectMixerEditorListFilter_Collection;
 class UObjectMixerEditorSerializedData;
 struct FObjectMixerEditorListRow;
@@ -28,6 +29,8 @@ public:
 	~FObjectMixerEditorMainPanel() = default;
 
 	void Init();
+	
+	void RegisterAndMapContextMenuCommands();
 
 	TSharedRef<SWidget> GetOrCreateWidget();
 
@@ -44,6 +47,9 @@ public:
 	 */
 	void RefreshList() const;
 
+	/** Called when the Rename command is executed from the UI or hotkey. */
+	void OnRenameCommand();
+
 	void RequestSyncEditorSelectionToListSelection() const;
 
 	TWeakPtr<FObjectMixerEditorList> GetEditorListModel() const
@@ -156,6 +162,8 @@ public:
 	FOnPreFilterChange OnPreFilterChange;
 	FOnPostFilterChange OnPostFilterChange;
 
+	TSharedPtr<FUICommandList> ObjectMixerCommands;
+
 private:
 
 	TSharedPtr<SObjectMixerEditorMainPanel> MainPanelWidget;
diff --git a/Engine/Plugins/Editor/ObjectMixer/ObjectMixer/Source/ObjectMixer/Public/Views/MainPanel/SObjectMixerEditorMainPanel.h b/Engine/Plugins/Editor/ObjectMixer/ObjectMixer/Source/ObjectMixer/Public/Views/MainPanel/SObjectMixerEditorMainPanel.h
--- a/Engine/Plugins/Editor/ObjectMixer/ObjectMixer/Source/ObjectMixer/Public/Views/MainPanel/SObjectMixerEditorMainPanel.h
+++ b/Engine/Plugins/Editor/ObjectMixer/ObjectMixer/Source/ObjectMixer/Public/Views/MainPanel/SObjectMixerEditorMainPanel.h
@@ -22,6 +22,8 @@ public:
 
 	void Construct(const FArguments& InArgs, const TSharedRef<FObjectMixerEditorMainPanel>& InMainPanel);
 
+	virtual FReply OnKeyDown( const FGeometry& MyGeometry, const FKeyEvent& InKeyEvent ) override;
+
 	TWeakPtr<FObjectMixerEditorMainPanel> GetMainPanelModel()
 	{
 		return MainPanelModel;
diff --git a/Engine/Plugins/Editor/ObjectMixer/ObjectMixer/Source/ObjectMixer/Public/Views/Widgets/ObjectMixerEditorListMenuContext.h b/Engine/Plugins/Editor/ObjectMixer/ObjectMixer/Source/ObjectMixer/Public/Views/Widgets/ObjectMixerEditorListMenuContext.h
--- a/Engine/Plugins/Editor/ObjectMixer/ObjectMixer/Source/ObjectMixer/Public/Views/Widgets/ObjectMixerEditorListMenuContext.h
+++ b/Engine/Plugins/Editor/ObjectMixer/ObjectMixer/Source/ObjectMixer/Public/Views/Widgets/ObjectMixerEditorListMenuContext.h
@@ -32,10 +32,13 @@ public:
 
 private:
 
+	static void CreateSelectCollectionsSubMenu(UToolMenu* Menu, FObjectMixerEditorListMenuContextData ContextData);
 	static void OnTextCommitted(const FText& InText, ETextCommit::Type InCommitType, const FObjectMixerEditorListMenuContextData ContextData);
 	
 	static void OnClickCollectionMenuEntry(const FName Key, const FObjectMixerEditorListMenuContextData ContextData);
 	static void AddObjectsToCollection(const FName Key, const FObjectMixerEditorListMenuContextData ContextData);
 	static void RemoveObjectsFromCollection(const FName Key, const FObjectMixerEditorListMenuContextData ContextData);
 	static bool AreAllObjectsInCollection(const FName Key, const FObjectMixerEditorListMenuContextData ContextData);
+
+	static void ReplaceEditSubMenu(const FObjectMixerEditorListMenuContextData& ContextData);
 };
