Object Mixer:

-Create Folder functionality
-Drag and drop actors into folder rows
-Add Move To Context Menu Option
-Refactored inline renaming to allow for renames to be prompted on next frame (for creating new folders)

diff --git a/Engine/Plugins/Editor/ObjectMixer/ObjectMixer/Source/ObjectMixer/ObjectMixerEditor.build.cs b/Engine/Plugins/Editor/ObjectMixer/ObjectMixer/Source/ObjectMixer/ObjectMixerEditor.build.cs
--- a/Engine/Plugins/Editor/ObjectMixer/ObjectMixer/Source/ObjectMixer/ObjectMixerEditor.build.cs
+++ b/Engine/Plugins/Editor/ObjectMixer/ObjectMixer/Source/ObjectMixer/ObjectMixerEditor.build.cs
@@ -14,7 +14,8 @@ public class ObjectMixerEditor : ModuleRules
 			{
 				"Core",
 				"OutputLog",
-				"PropertyEditor"
+				"PropertyEditor", 
+				"SceneOutliner"
 			}
 		);
 
diff --git a/Engine/Plugins/Editor/ObjectMixer/ObjectMixer/Source/ObjectMixer/Private/ObjectMixerEditorModule.cpp b/Engine/Plugins/Editor/ObjectMixer/ObjectMixer/Source/ObjectMixer/Private/ObjectMixerEditorModule.cpp
--- a/Engine/Plugins/Editor/ObjectMixer/ObjectMixer/Source/ObjectMixer/Private/ObjectMixerEditorModule.cpp
+++ b/Engine/Plugins/Editor/ObjectMixer/ObjectMixer/Source/ObjectMixer/Private/ObjectMixerEditorModule.cpp
@@ -45,6 +45,12 @@ void FObjectMixerEditorModule::ShutdownModule()
 	Teardown();
 }
 
+UWorld* FObjectMixerEditorModule::GetWorld()
+{
+	check(GEditor);
+	return GEditor->GetEditorWorldContext().World();
+}
+
 void FObjectMixerEditorModule::Initialize()
 {
 	BindDelegates();
diff --git a/Engine/Plugins/Editor/ObjectMixer/ObjectMixer/Source/ObjectMixer/Private/Views/List/ObjectMixerEditorList.cpp b/Engine/Plugins/Editor/ObjectMixer/ObjectMixer/Source/ObjectMixer/Private/Views/List/ObjectMixerEditorList.cpp
--- a/Engine/Plugins/Editor/ObjectMixer/ObjectMixer/Source/ObjectMixer/Private/Views/List/ObjectMixerEditorList.cpp
+++ b/Engine/Plugins/Editor/ObjectMixer/ObjectMixer/Source/ObjectMixer/Private/Views/List/ObjectMixerEditorList.cpp
@@ -103,6 +103,22 @@ void FObjectMixerEditorList::RequestSyncEditorSelectionToListSelection() const
 	}
 }
 
+void FObjectMixerEditorList::OnRequestNewFolder(TOptional<FFolder> ExplicitParentFolder)
+{
+	if (ListWidget.IsValid())
+	{
+		ListWidget->OnRequestNewFolder(ExplicitParentFolder);
+	}
+}
+
+void FObjectMixerEditorList::OnRequestMoveFolder(const FFolder& FolderToMove, const FFolder& TargetNewParentFolder)
+{
+	if (ListWidget.IsValid())
+	{
+		ListWidget->OnRequestMoveFolder(FolderToMove, TargetNewParentFolder);
+	}
+}
+
 void FObjectMixerEditorList::ExecuteListViewSearchOnAllRows(const FString& SearchString,
                                                             const bool bShouldRefreshAfterward)
 {
diff --git a/Engine/Plugins/Editor/ObjectMixer/ObjectMixer/Source/ObjectMixer/Private/Views/List/ObjectMixerEditorListRow.cpp b/Engine/Plugins/Editor/ObjectMixer/ObjectMixer/Source/ObjectMixer/Private/Views/List/ObjectMixerEditorListRow.cpp
--- a/Engine/Plugins/Editor/ObjectMixer/ObjectMixer/Source/ObjectMixer/Private/Views/List/ObjectMixerEditorListRow.cpp
+++ b/Engine/Plugins/Editor/ObjectMixer/ObjectMixer/Source/ObjectMixer/Private/Views/List/ObjectMixerEditorListRow.cpp
@@ -183,7 +183,7 @@ void FObjectMixerEditorListRow::SetChildRows(const TArray<FObjectMixerEditorList
 
 void FObjectMixerEditorListRow::AddToChildRows(const FObjectMixerEditorListRowPtr& InRow)
 {
-	InRow->SetDirectParentRow(SharedThis(this));
+	InRow->SetDirectParentRow(GetAsShared());
 	ChildRows.AddUnique(InRow);
 	ChildRows.StableSort(SObjectMixerEditorList::SortByTypeThenName);
 }
@@ -194,6 +194,30 @@ void FObjectMixerEditorListRow::InsertChildRowAtIndex(const FObjectMixerEditorLi
 	ChildRows.Insert(InRow, AtIndex);
 }
 
+void FObjectMixerEditorListRow::SetChildRowsSelected(const bool bNewSelected, const bool bRecursive,
+	const bool bSelectOnlyVisible)
+{
+	for (const FObjectMixerEditorListRowPtr& ChildRow : GetChildRows())
+	{
+		if (ChildRow.IsValid())
+		{
+			// Recurse even if not visible
+			if (bRecursive)
+			{
+				ChildRow->SetChildRowsSelected(bNewSelected, bRecursive, bSelectOnlyVisible);
+			}
+
+			// Skip setting selection if not visible and bSelectOnlyVisible == true
+			if (bSelectOnlyVisible && !ChildRow->ShouldRowWidgetBeVisible())
+			{
+				continue;
+			}
+		
+			ChildRow->SetIsSelected(bNewSelected);
+		}
+	}
+}
+
 bool FObjectMixerEditorListRow::GetIsTreeViewItemExpanded()
 {
 	return GetListViewPtr().Pin()->IsTreeViewItemExpanded(GetAsShared());
@@ -312,7 +336,14 @@ bool FObjectMixerEditorListRow::GetIsSelected()
 {
 	check (ListViewPtr.IsValid());
 
-	return ListViewPtr.Pin()->IsTreeViewItemSelected(SharedThis(this));
+	return ListViewPtr.Pin()->IsTreeViewItemSelected(GetAsShared().ToSharedRef());
+}
+
+void FObjectMixerEditorListRow::SetIsSelected(const bool bNewSelected)
+{
+	check (ListViewPtr.IsValid());
+
+	return ListViewPtr.Pin()->SetTreeViewItemSelected(SharedThis(this), bNewSelected);
 }
 
 bool FObjectMixerEditorListRow::ShouldRowWidgetBeVisible() const
diff --git a/Engine/Plugins/Editor/ObjectMixer/ObjectMixer/Source/ObjectMixer/Private/Views/List/SObjectMixerEditorList.cpp b/Engine/Plugins/Editor/ObjectMixer/ObjectMixer/Source/ObjectMixer/Private/Views/List/SObjectMixerEditorList.cpp
--- a/Engine/Plugins/Editor/ObjectMixer/ObjectMixer/Source/ObjectMixer/Private/Views/List/SObjectMixerEditorList.cpp
+++ b/Engine/Plugins/Editor/ObjectMixer/ObjectMixer/Source/ObjectMixer/Private/Views/List/SObjectMixerEditorList.cpp
@@ -18,11 +18,13 @@
 #include "DragAndDrop/AssetDragDropOp.h"
 #include "Editor.h"
 #include "Editor/UnrealEdEngine.h"
+#include "EditorActorFolders.h"
 #include "Framework/MultiBox/MultiBoxBuilder.h"
 #include "IPlacementModeModule.h"
 #include "LevelEditor.h"
 #include "LevelEditorMenuContext.h"
 #include "LevelEditorViewport.h"
+#include "ObjectMixerEditorModule.h"
 #include "ObjectMixerEditorSerializedData.h"
 #include "ScopedTransaction.h"
 #include "Selection.h"
@@ -108,14 +110,7 @@ void SObjectMixerEditorList::Construct(const FArguments& InArgs, TSharedRef<FObj
 					{
 						if (GCurrentLevelEditingViewportClient)
 						{
-							AActor* Actor = Cast<AActor>(Row->GetObject());
-
-							if (!Actor)
-							{
-								Actor = Row->GetObject()->GetTypedOuter<AActor>();
-							}
-
-							if (Actor)
+							if (const AActor* Actor = Row->GetSelfOrOuterAsActor())
 							{
 								FVector Origin;
 								FVector Extents;
@@ -170,6 +165,11 @@ SObjectMixerEditorList::~SObjectMixerEditorList()
 	TreeViewPtr.Reset();
 }
 
+UWorld* SObjectMixerEditorList::GetWorld() const
+{
+	return FObjectMixerEditorModule::Get().GetWorld();
+}
+
 void SObjectMixerEditorList::RebuildList()
 {
 	bIsRebuildRequested = false;
@@ -204,7 +204,9 @@ void SObjectMixerEditorList::OnRenameCommand()
 {
 	if (GetSelectedTreeViewItemCount() == 1)
 	{
-		GetSelectedTreeViewItems()[0]->CallOnRenameCommandDelegate();
+		const FObjectMixerEditorListRowPtr TreeViewItem = GetSelectedTreeViewItems()[0];
+
+		PendingRenameItem = FTreeItemUniqueIdentifier(TreeViewItem);
 	}
 }
 
@@ -218,6 +220,11 @@ int32 SObjectMixerEditorList::GetSelectedTreeViewItemCount() const
 	return TreeViewPtr->GetSelectedItems().Num();
 }
 
+void SObjectMixerEditorList::SetTreeViewItemSelected(FObjectMixerEditorListRowPtr Item, const bool bNewSelected)
+{
+	TreeViewPtr->SetItemSelection(Item, bNewSelected);
+}
+
 void SObjectMixerEditorList::SyncEditorSelectionToListSelection()
 {
 	if (GetDefault<UObjectMixerEditorSettings>()->bSyncSelection && !bShouldPauseSyncSelection)
@@ -252,6 +259,67 @@ void SObjectMixerEditorList::SyncEditorSelectionToListSelection()
 	}
 }
 
+void SObjectMixerEditorList::OnRequestNewFolder(TOptional<FFolder> ExplicitParentFolder)
+{
+	const FScopedTransaction Transaction(LOCTEXT("UndoAction_CreateFolder", "Create Folder"));
+
+	UWorld* World = GetWorld();
+	FFolder NewFolderPath;
+	if (ExplicitParentFolder.IsSet())
+	{
+		NewFolderPath =
+			FActorFolders::Get().GetDefaultFolderName(*World, ExplicitParentFolder.GetValue());
+	}
+	else
+	{
+		NewFolderPath = FActorFolders::Get().GetDefaultFolderForSelection(*World);
+	}
+
+	if (FActorFolders::Get().CreateFolder(*World, NewFolderPath))
+	{
+		const TArray<FObjectMixerEditorListRowPtr>& SelectedTreeItems = TreeViewPtr->GetSelectedItems();
+
+		if (SelectedTreeItems.Num())
+		{
+			TArray<AActor*> ActorsToMoveToFolder;
+			for (const FObjectMixerEditorListRowPtr& SelectedRow : SelectedTreeItems)
+			{
+				if (SelectedRow->GetRowType() == FObjectMixerEditorListRow::MatchingObject ||
+					SelectedRow->GetRowType() == FObjectMixerEditorListRow::ContainerObject)
+				{
+					if (AActor* Actor = SelectedRow->GetSelfOrOuterAsActor())
+					{
+						ActorsToMoveToFolder.Add(Actor);
+					}
+				}
+			}
+
+			if (ActorsToMoveToFolder.Num())
+			{				
+				for (AActor* Actor : ActorsToMoveToFolder)
+				{
+					Actor->Modify();
+					Actor->SetFolderPath_Recursively(NewFolderPath.GetPath());
+				}
+			}
+		}
+
+		// Set new folder row to be renamed on next frame
+		PendingRenameItem.RowName = NewFolderPath.GetPath().ToString();
+	}
+}
+
+void SObjectMixerEditorList::OnRequestMoveFolder(const FFolder& FolderToMove, const FFolder& TargetNewParentFolder)
+{
+	if (UWorld* EditorWorld = GetWorld())
+	{
+		check(TargetNewParentFolder.GetRootObject() == FolderToMove.GetRootObject());
+		// Get unique name
+		const FFolder NewPath = FActorFolders::Get().GetFolderName(*EditorWorld, TargetNewParentFolder, FolderToMove.GetLeafName());
+		FActorFolders::Get().RenameFolderInWorld(*EditorWorld, FolderToMove, NewPath);
+	}
+}
+
 void SObjectMixerEditorList::SetSelectedTreeViewItemActorsEditorVisible(const bool bNewIsVisible, const bool bIsRecursive)
 {
 	for (const FObjectMixerEditorListRowPtr& SelectedItem : TreeViewPtr->GetSelectedItems())
@@ -606,6 +674,41 @@ FListViewColumnInfo* SObjectMixerEditorList::GetColumnInfoByPropertyName(const F
 		});
 }
 
+void SObjectMixerEditorList::ExecuteRenameOnPending()
+{
+	using LambdaType = const FObjectMixerEditorListRowPtr*(*)(const TArray<FObjectMixerEditorListRowPtr>&, const FTreeItemUniqueIdentifier&);
+		
+	static LambdaType RecursivelyFindByUniqueIdentifier = [](
+		const TArray<FObjectMixerEditorListRowPtr>& InObjects, const FTreeItemUniqueIdentifier& InUniqueIdentifier)
+	{
+		for (const FObjectMixerEditorListRowPtr& TreeViewItem : InObjects)
+		{
+			FTreeItemUniqueIdentifier TreeItemUniqueIdentifier(TreeViewItem);
+
+			if (InUniqueIdentifier == TreeItemUniqueIdentifier)
+			{
+				return &TreeViewItem;
+			}
+
+			if (const FObjectMixerEditorListRowPtr* FoundItem = RecursivelyFindByUniqueIdentifier(
+				TreeViewItem->GetChildRows(), InUniqueIdentifier))
+			{
+				return FoundItem;
+			}
+		}
+
+		const FObjectMixerEditorListRowPtr* Nulled = nullptr;
+		return Nulled;
+	};
+
+	if (const FObjectMixerEditorListRowPtr* FoundItem = RecursivelyFindByUniqueIdentifier(TreeViewRootObjects, PendingRenameItem))
+	{
+		(*FoundItem)->CallOnRenameCommandDelegate();
+	}
+
+	PendingRenameItem.Reset();
+}
+
 void SObjectMixerEditorList::Tick(const FGeometry& AllottedGeometry, const double InCurrentTime,
 	const float InDeltaTime)
 {
@@ -614,6 +717,9 @@ void SObjectMixerEditorList::Tick(const FGeometry& AllottedGeometry, const doubl
 	if (bIsRebuildRequested)
 	{
 		RebuildList();
+		
+		// Early out so we have widgets to act upon next frame
+		return;
 	}
 
 	if (bIsEditorToListSelectionSyncRequested)
@@ -626,6 +732,11 @@ void SObjectMixerEditorList::Tick(const FGeometry& AllottedGeometry, const doubl
 		bShouldPauseSyncSelection = false;
 		bIsEditorToListSelectionSyncRequested = false;
 	}
+
+	if (PendingRenameItem.IdentifiesAnyRow())
+	{
+		ExecuteRenameOnPending();
+	}
 }
 
 TSharedRef<SWidget> SObjectMixerEditorList::GenerateHeaderRowContextMenu() const
@@ -998,17 +1109,13 @@ void SObjectMixerEditorList::CacheTreeState(const TArray<TWeakPtr<IObjectMixerEd
 	{
 		for (const FObjectMixerEditorListRowPtr& TreeViewItem : InObjects)
 		{
-			UObject* RowObject = TreeViewItem->GetObject();
-			const FString RowName =
-				TreeViewItem->GetRowType() == FObjectMixerEditorListRow::Folder ?
-					TreeViewItem->GetFolderPath().ToString() : TreeViewItem->GetDisplayName().ToString();
+			FTreeItemUniqueIdentifier UniqueIdentifier(TreeViewItem);
 
-			if (!RowName.IsEmpty())
+			if (!UniqueIdentifier.RowName.IsEmpty())
 			{
 				TreeItemStateCache->Add(
 					{
-						RowObject ? RowObject->GetUniqueID() : -1,
-						RowName,
+						UniqueIdentifier,
 						InTreeViewPtr->IsItemExpanded(TreeViewItem),
 						InTreeViewPtr->IsItemSelected(TreeViewItem),
 						TreeViewItem->GetVisibilityRules()
@@ -1055,23 +1162,14 @@ void SObjectMixerEditorList::RestoreTreeState(const TArray<TWeakPtr<IObjectMixer
 	{
 		for (const FObjectMixerEditorListRowPtr& TreeViewItem : InObjects)
 		{
+			const FTreeItemUniqueIdentifier TreeItemUniqueIdentifier(TreeViewItem);
+			
 			if (const FTreeItemStateCache* StateCachePtr = Algo::FindByPredicate(
 				*TreeItemStateCache,
-				[TreeViewItem](const FTreeItemStateCache& Other)
+				[&TreeItemUniqueIdentifier](const FTreeItemStateCache& Other)
 				{
-					if (Other.UniqueId != -1)
-					{
-						if (const UObject* RowObject = TreeViewItem->GetObject())
-						{
-							return Other.UniqueId == RowObject->GetUniqueID();
-						}
-					}
-
-					const FString RowName =
-						TreeViewItem->GetRowType() == FObjectMixerEditorListRow::Folder ?
-							TreeViewItem->GetFolderPath().ToString() : TreeViewItem->GetDisplayName().ToString();
-
-					return Other.RowName.Equals(RowName);
+					
+					return Other.UniqueIdentifier == TreeItemUniqueIdentifier;
 				}
 			))
 			{
@@ -1343,8 +1441,7 @@ void SObjectMixerEditorList::GenerateTreeView()
 
 	BuildPerformanceCacheAndGenerateHeaderIfNeeded();
 
-	check(GEditor);
-	const UWorld* EditorWorld = GEditor->GetEditorWorldContext().World();
+	const UWorld* EditorWorld = GetWorld();
 	
 	TSet<UObject*> AllMatchingObjects;
 
@@ -1420,7 +1517,6 @@ void SObjectMixerEditorList::SelectedTreeItemsToSelectedInLevelEditor()
 	if (GEditor && !bShouldPauseSyncSelection)
 	{
 		bShouldPauseSyncSelection = true;
-		TArray<AActor*> ActorsToSelect;
 		const TArray<FObjectMixerEditorListRowPtr>& SelectedTreeItems = TreeViewPtr->GetSelectedItems();
 		
 		if (SelectedTreeItems.Num() == 0)
@@ -1429,19 +1525,13 @@ void SObjectMixerEditorList::SelectedTreeItemsToSelectedInLevelEditor()
 		}
 		else
 		{
+			TArray<AActor*> ActorsToSelect;
 			for (const FObjectMixerEditorListRowPtr& SelectedRow : SelectedTreeItems)
 			{
 				if (SelectedRow->GetRowType() == FObjectMixerEditorListRow::MatchingObject ||
 					SelectedRow->GetRowType() == FObjectMixerEditorListRow::ContainerObject)
 				{
-					AActor* Actor = Cast<AActor>(SelectedRow->GetObject());
-
-					if (!Actor)
-					{
-						Actor = SelectedRow->GetObject()->GetTypedOuter<AActor>();
-					}
-
-					if (Actor)
+					if (AActor* Actor = SelectedRow->GetSelfOrOuterAsActor())
 					{
 						ActorsToSelect.Add(Actor);
 					}
diff --git a/Engine/Plugins/Editor/ObjectMixer/ObjectMixer/Source/ObjectMixer/Private/Views/List/SObjectMixerEditorListRow.cpp b/Engine/Plugins/Editor/ObjectMixer/ObjectMixer/Source/ObjectMixer/Private/Views/List/SObjectMixerEditorListRow.cpp
--- a/Engine/Plugins/Editor/ObjectMixer/ObjectMixer/Source/ObjectMixer/Private/Views/List/SObjectMixerEditorListRow.cpp
+++ b/Engine/Plugins/Editor/ObjectMixer/ObjectMixer/Source/ObjectMixer/Private/Views/List/SObjectMixerEditorListRow.cpp
@@ -548,10 +548,6 @@ TOptional<EItemDropZone> SObjectMixerEditorListRow::HandleCanAcceptDrop(
 			Operation->DraggedItems.Num() == 1 &&
 			Operation->DraggedItems[0]->GetObject() &&
 			Operation->DraggedItems[0]->GetObject()->IsA(UActorComponent::StaticClass())
-		) ||
-		(
-			Operation->DraggedItems.Num() == 1 &&
-			Operation->DraggedItems[0]->GetRowType() == FObjectMixerEditorListRow::Folder
 		)
 	;
 
@@ -612,7 +608,17 @@ FReply SObjectMixerEditorListRow::HandleAcceptDrop(const FDragDropEvent& DragDro
 
 	for (const FObjectMixerEditorListRowPtr& DraggedItem : Operation->DraggedItems)
 	{
-		if (AActor* ObjectAsActor = Cast<AActor>(DraggedItem->GetObject()))
+		if (DraggedItem->GetRowType() == FObjectMixerEditorListRow::Folder)
+		{
+			if (bIsDroppingOnFolderRow)
+			{
+				if (const TSharedPtr<SObjectMixerEditorList> PinnedList = DraggedItem->GetListViewPtr().Pin())
+				{
+					PinnedList->OnRequestMoveFolder(DraggedItem->GetFolder(), TargetItem->GetFolder());
+				}
+			}
+		}
+		else if (AActor* ObjectAsActor = Cast<AActor>(DraggedItem->GetObject()))
 		{
 			if (bIsDroppingOnFolderRow)
 			{
diff --git a/Engine/Plugins/Editor/ObjectMixer/ObjectMixer/Source/ObjectMixer/Private/Views/MainPanel/ObjectMixerEditorMainPanel.cpp b/Engine/Plugins/Editor/ObjectMixer/ObjectMixer/Source/ObjectMixer/Private/Views/MainPanel/ObjectMixerEditorMainPanel.cpp
--- a/Engine/Plugins/Editor/ObjectMixer/ObjectMixer/Source/ObjectMixer/Private/Views/MainPanel/ObjectMixerEditorMainPanel.cpp
+++ b/Engine/Plugins/Editor/ObjectMixer/ObjectMixer/Source/ObjectMixer/Private/Views/MainPanel/ObjectMixerEditorMainPanel.cpp
@@ -17,29 +17,55 @@ void FObjectMixerEditorMainPanel::Init()
 
 void FObjectMixerEditorMainPanel::RegisterAndMapContextMenuCommands()
 {
-	ObjectMixerCommands = MakeShared<FUICommandList>();
+	ObjectMixerElementEditCommands = MakeShared<FUICommandList>();
 
-	ObjectMixerCommands->MapAction(FGenericCommands::Get().Cut,
+	ObjectMixerElementEditCommands->MapAction(FGenericCommands::Get().Cut,
 		FExecuteAction::CreateStatic( &FLevelEditorActionCallbacks::ExecuteExecCommand, FString( TEXT("EDIT CUT") ) ),
 		FCanExecuteAction::CreateStatic( &FLevelEditorActionCallbacks::Cut_CanExecute )
 	);
-	ObjectMixerCommands->MapAction(FGenericCommands::Get().Copy,
+	ObjectMixerElementEditCommands->MapAction(FGenericCommands::Get().Copy,
 		FExecuteAction::CreateStatic( &FLevelEditorActionCallbacks::ExecuteExecCommand, FString( TEXT("EDIT COPY") ) ),
 		FCanExecuteAction::CreateStatic( &FLevelEditorActionCallbacks::Copy_CanExecute )
 	);
-	ObjectMixerCommands->MapAction(FGenericCommands::Get().Paste,
+	ObjectMixerElementEditCommands->MapAction(FGenericCommands::Get().Paste,
 		FExecuteAction::CreateStatic( &FLevelEditorActionCallbacks::ExecuteExecCommand, FString( TEXT("EDIT PASTE") ) ),
 		FCanExecuteAction::CreateStatic( &FLevelEditorActionCallbacks::Paste_CanExecute )
 	);
-	ObjectMixerCommands->MapAction(FGenericCommands::Get().Duplicate,
+	ObjectMixerElementEditCommands->MapAction(FGenericCommands::Get().Duplicate,
 		FExecuteAction::CreateStatic( &FLevelEditorActionCallbacks::ExecuteExecCommand, FString( TEXT("DUPLICATE") ) ),
 		FCanExecuteAction::CreateStatic( &FLevelEditorActionCallbacks::Duplicate_CanExecute )
 	);
-	ObjectMixerCommands->MapAction(FGenericCommands::Get().Delete,
+	ObjectMixerElementEditCommands->MapAction(FGenericCommands::Get().Delete,
 		FExecuteAction::CreateStatic( &FLevelEditorActionCallbacks::ExecuteExecCommand, FString( TEXT("DELETE") ) ),
 		FCanExecuteAction::CreateStatic( &FLevelEditorActionCallbacks::Delete_CanExecute )
 	);
-	ObjectMixerCommands->MapAction(FGenericCommands::Get().Rename,
+	ObjectMixerElementEditCommands->MapAction(FGenericCommands::Get().Rename,
+		FUIAction(FExecuteAction::CreateRaw(this, &FObjectMixerEditorMainPanel::OnRenameCommand))
+	);
+
+	ObjectMixerFolderEditCommands = MakeShared<FUICommandList>();
+
+	ObjectMixerElementEditCommands->MapAction(FGenericCommands::Get().Cut,
+		FExecuteAction::CreateStatic( &FLevelEditorActionCallbacks::ExecuteExecCommand, FString( TEXT("EDIT CUT") ) ),
+		FCanExecuteAction::CreateStatic( &FLevelEditorActionCallbacks::Cut_CanExecute )
+	);
+	ObjectMixerElementEditCommands->MapAction(FGenericCommands::Get().Copy,
+		FExecuteAction::CreateStatic( &FLevelEditorActionCallbacks::ExecuteExecCommand, FString( TEXT("EDIT COPY") ) ),
+		FCanExecuteAction::CreateStatic( &FLevelEditorActionCallbacks::Copy_CanExecute )
+	);
+	ObjectMixerElementEditCommands->MapAction(FGenericCommands::Get().Paste,
+		FExecuteAction::CreateStatic( &FLevelEditorActionCallbacks::ExecuteExecCommand, FString( TEXT("EDIT PASTE") ) ),
+		FCanExecuteAction::CreateStatic( &FLevelEditorActionCallbacks::Paste_CanExecute )
+	);
+	ObjectMixerElementEditCommands->MapAction(FGenericCommands::Get().Duplicate,
+		FExecuteAction::CreateStatic( &FLevelEditorActionCallbacks::ExecuteExecCommand, FString( TEXT("DUPLICATE") ) ),
+		FCanExecuteAction::CreateStatic( &FLevelEditorActionCallbacks::Duplicate_CanExecute )
+	);
+	ObjectMixerElementEditCommands->MapAction(FGenericCommands::Get().Delete,
+		FExecuteAction::CreateStatic( &FLevelEditorActionCallbacks::ExecuteExecCommand, FString( TEXT("DELETE") ) ),
+		FCanExecuteAction::CreateStatic( &FLevelEditorActionCallbacks::Delete_CanExecute )
+	);
+	ObjectMixerElementEditCommands->MapAction(FGenericCommands::Get().Rename,
 		FUIAction(FExecuteAction::CreateRaw(this, &FObjectMixerEditorMainPanel::OnRenameCommand))
 	);
 }
@@ -85,6 +111,22 @@ void FObjectMixerEditorMainPanel::OnRenameCommand()
 	}
 }
 
+void FObjectMixerEditorMainPanel::OnRequestNewFolder(TOptional<FFolder> ExplicitParentFolder)
+{
+	if (EditorListModel.IsValid())
+	{
+		EditorListModel->OnRequestNewFolder(ExplicitParentFolder);
+	}
+}
+
+void FObjectMixerEditorMainPanel::OnRequestMoveFolder(const FFolder& FolderToMove, const FFolder& TargetNewParentFolder)
+{
+	if (EditorListModel.IsValid())
+	{
+		EditorListModel->OnRequestMoveFolder(FolderToMove, TargetNewParentFolder);
+	}
+}
+
 void FObjectMixerEditorMainPanel::RequestSyncEditorSelectionToListSelection() const
 {
 	if (EditorListModel.IsValid())
diff --git a/Engine/Plugins/Editor/ObjectMixer/ObjectMixer/Source/ObjectMixer/Private/Views/MainPanel/SObjectMixerEditorMainPanel.cpp b/Engine/Plugins/Editor/ObjectMixer/ObjectMixer/Source/ObjectMixer/Private/Views/MainPanel/SObjectMixerEditorMainPanel.cpp
--- a/Engine/Plugins/Editor/ObjectMixer/ObjectMixer/Source/ObjectMixer/Private/Views/MainPanel/SObjectMixerEditorMainPanel.cpp
+++ b/Engine/Plugins/Editor/ObjectMixer/ObjectMixer/Source/ObjectMixer/Private/Views/MainPanel/SObjectMixerEditorMainPanel.cpp
@@ -16,11 +16,13 @@
 #include "AssetRegistry/AssetRegistryModule.h"
 #include "AssetRegistry/IAssetRegistry.h"
 #include "Framework/Application/SlateApplication.h"
+#include "EditorActorFolders.h"
 #include "Framework/MultiBox/MultiBoxBuilder.h"
 #include "IPlacementModeModule.h"
 #include "Kismet2/SClassPickerDialog.h"
 #include "SPositiveActionButton.h"
 #include "Styling/StyleColors.h"
+#include "Widgets/Input/SButton.h"
 #include "Widgets/Input/SCheckBox.h"
 #include "Widgets/Input/SComboButton.h"
 #include "Widgets/Input/SMultiLineEditableTextBox.h"
@@ -88,7 +90,7 @@ FReply SObjectMixerEditorMainPanel::OnKeyDown(const FGeometry& MyGeometry, const
 	// If we are in debug mode do not process commands
 	if (FSlateApplication::Get().IsNormalExecution())
 	{
-		if (GetMainPanelModel().Pin()->ObjectMixerCommands->ProcessCommandBindings(InKeyEvent))
+		if (GetMainPanelModel().Pin()->ObjectMixerElementEditCommands->ProcessCommandBindings(InKeyEvent))
 		{
 			return FReply::Handled();
 		}
@@ -158,6 +160,34 @@ TSharedRef<SWidget> SObjectMixerEditorMainPanel::GenerateToolbar()
 		]
 	];
 
+	// Create Folder
+	ToolbarBox->AddSlot()
+	.HAlign(HAlign_Right)
+	.VAlign(VAlign_Center)
+	.AutoWidth()
+	.Padding(8.f, 1.f, 0.f, 1.f)
+	[
+		SNew(SButton)
+		.ButtonStyle(FAppStyle::Get(), "SimpleButton")
+		.ToolTipText(LOCTEXT("CreateFolderToolTip", "Create a new folder containing the current selection"))
+		.OnClicked_Lambda([this] ()
+		{
+			if (TSharedPtr<FObjectMixerEditorList> PinnedList = GetMainPanelModel().Pin()->GetEditorListModel().Pin())
+			{
+				PinnedList->OnRequestNewFolder();
+
+				return FReply::Handled();
+			}
+
+			return FReply::Unhandled();
+		})
+		[
+			SNew(SImage)
+			.ColorAndOpacity(FSlateColor::UseForeground())
+			.Image(FAppStyle::Get().GetBrush("SceneOutliner.NewFolderIcon"))
+		]
+	];
+
 	// Show Options
 	ToolbarBox->AddSlot()
 	.HAlign(HAlign_Right)
diff --git a/Engine/Plugins/Editor/ObjectMixer/ObjectMixer/Source/ObjectMixer/Private/Views/Widgets/ObjectMixerEditorListMenuContext.cpp b/Engine/Plugins/Editor/ObjectMixer/ObjectMixer/Source/ObjectMixer/Private/Views/Widgets/ObjectMixerEditorListMenuContext.cpp
--- a/Engine/Plugins/Editor/ObjectMixer/ObjectMixer/Source/ObjectMixer/Private/Views/Widgets/ObjectMixerEditorListMenuContext.cpp
+++ b/Engine/Plugins/Editor/ObjectMixer/ObjectMixer/Source/ObjectMixer/Private/Views/Widgets/ObjectMixerEditorListMenuContext.cpp
@@ -2,10 +2,13 @@
 
 #include "Views/Widgets/ObjectMixerEditorListMenuContext.h"
 
+#include "ActorFolderPickingMode.h"
+#include "FolderTreeItem.h"
 #include "Views/MainPanel/ObjectMixerEditorMainPanel.h"
 
 #include "Elements/Framework/TypedElementSelectionSet.h"
 #include "Elements/Interfaces/TypedElementObjectInterface.h"
+#include "Framework/Application/SlateApplication.h"
 #include "Framework/MultiBox/MultiBoxDefs.h"
 #include "LevelEditor.h"
 #include "LevelEditorActions.h"
@@ -14,6 +17,7 @@
 #include "LevelEditorViewport.h"
 #include "ObjectMixerEditorModule.h"
 #include "Selection.h"
+#include "SSceneOutliner.h"
 #include "ToolMenu.h"
 #include "ToolMenus.h"
 #include "ToolMenuContext.h"
@@ -38,31 +42,41 @@ TSharedPtr<SWidget> UObjectMixerEditorListMenuContext::CreateContextMenu(const F
 	return BuildContextMenu(InData);
 }
 
-TSharedPtr<SWidget> UObjectMixerEditorListMenuContext::BuildContextMenu(const FObjectMixerEditorListMenuContextData& InData)
+bool UObjectMixerEditorListMenuContext::DoesSelectionHaveType(const FObjectMixerEditorListMenuContextData& InData, UClass* Type)
 {
-	FLevelEditorContextMenu::RegisterComponentContextMenu();
-	FLevelEditorContextMenu::RegisterActorContextMenu();
-	FLevelEditorContextMenu::RegisterElementContextMenu();
-	RegisterObjectMixerContextMenu();
+	if (const TSharedPtr<FObjectMixerEditorListRow>* Match = Algo::FindByPredicate(
+		InData.SelectedItems,
+		[Type](const TSharedPtr<FObjectMixerEditorListRow>& SelectedItem)
+		{
+			const UObject* Object = SelectedItem->GetObject();
+			return Object && Object->IsA(Type);
+		}))
+	{
+		return true;
+	}
+
+	return false;
+}
 
+void PerformLevelEditorRegistrations(FToolMenuContext& Context)
+{
 	FLevelEditorModule& LevelEditorModule = FModuleManager::LoadModuleChecked<FLevelEditorModule>("LevelEditor");
 	TSharedPtr<ILevelEditor> LevelEditorPtr = LevelEditorModule.GetLevelEditorInstance().Pin();
 	check(LevelEditorPtr);
 
 	TSharedPtr<FUICommandList> LevelEditorActionsList = LevelEditorPtr->GetLevelEditorActions();
-	FToolMenuContext Context;
 	Context.AppendCommandList(LevelEditorActionsList);
 
 	ULevelEditorContextMenuContext* LevelEditorContextObject = NewObject<ULevelEditorContextMenuContext>();
 	LevelEditorContextObject->LevelEditor = LevelEditorPtr;
 	LevelEditorContextObject->ContextType = ELevelEditorMenuContext::SceneOutliner;
 	LevelEditorContextObject->CurrentSelection = LevelEditorPtr->GetElementSelectionSet();
-	
+
 	for (FSelectedEditableComponentIterator It(GEditor->GetSelectedEditableComponentIterator()); It; ++It)
 	{
 		LevelEditorContextObject->SelectedComponents.Add(CastChecked<UActorComponent>(*It));
 	}
-	
+
 	Context.AddObject(LevelEditorContextObject, [](UObject* InContext)
 	{
 		ULevelEditorContextMenuContext* CastContext = CastChecked<ULevelEditorContextMenuContext>(InContext);
@@ -70,23 +84,14 @@ TSharedPtr<SWidget> UObjectMixerEditorListMenuContext::BuildContextMenu(const FO
 		CastContext->HitProxyElement.Release();
 	});
 
-	UObjectMixerEditorListMenuContext* ObjectMixerContextObject = NewObject<UObjectMixerEditorListMenuContext>();
-	ObjectMixerContextObject->Data = InData;
-
-	Context.AddObject(ObjectMixerContextObject, [](UObject* InContext)
-	{
-		UObjectMixerEditorListMenuContext* CastContext = CastChecked<UObjectMixerEditorListMenuContext>(InContext);
-		CastContext->Data.SelectedItems.Empty();
-		CastContext->Data.MainPanelPtr.Reset();
-	});
-
 	if (LevelEditorPtr->GetElementSelectionSet()->GetSelectedObjects<UActorComponent>().Num() == 0 &&
 		LevelEditorPtr->GetElementSelectionSet()->GetSelectedObjects<AActor>().Num() > 0)
 	{
 		const TArray<AActor*> SelectedActors = LevelEditorPtr->GetElementSelectionSet()->GetSelectedObjects<AActor>();
-	
+
 		// Get all menu extenders for this context menu from the level editor module
-		TArray<FLevelEditorModule::FLevelViewportMenuExtender_SelectedActors> MenuExtenderDelegates = LevelEditorModule.GetAllLevelViewportContextMenuExtenders();
+		TArray<FLevelEditorModule::FLevelViewportMenuExtender_SelectedActors> MenuExtenderDelegates =
+			LevelEditorModule.GetAllLevelViewportContextMenuExtenders();
 		TArray<TSharedPtr<FExtender>> Extenders;
 		for (int32 i = 0; i < MenuExtenderDelegates.Num(); ++i)
 		{
@@ -95,17 +100,156 @@ TSharedPtr<SWidget> UObjectMixerEditorListMenuContext::BuildContextMenu(const FO
 				Extenders.Add(MenuExtenderDelegates[i].Execute(LevelEditorActionsList.ToSharedRef(), SelectedActors));
 			}
 		}
-	
+
 		if (Extenders.Num() > 0)
 		{
 			Context.AddExtender(FExtender::Combine(Extenders));
 		}
 	}
+}
+
+TSharedPtr<SWidget> UObjectMixerEditorListMenuContext::BuildContextMenu(const FObjectMixerEditorListMenuContextData& InData)
+{	
+	FToolMenuContext Context;
+
+	UObjectMixerEditorListMenuContext* ObjectMixerContextObject = NewObject<UObjectMixerEditorListMenuContext>();
+	ObjectMixerContextObject->Data = InData;
+
+	Context.AddObject(ObjectMixerContextObject, [](UObject* InContext)
+	{
+		UObjectMixerEditorListMenuContext* CastContext = CastChecked<UObjectMixerEditorListMenuContext>(InContext);
+		CastContext->Data.SelectedItems.Empty();
+		CastContext->Data.MainPanelPtr.Reset();
+	});
+
+	if (DoesSelectionHaveType(InData, AActor::StaticClass()))
+	{
+		FLevelEditorContextMenu::RegisterActorContextMenu();
+		FLevelEditorContextMenu::RegisterElementContextMenu();
+		RegisterObjectMixerActorContextMenuExtension();
+
+		PerformLevelEditorRegistrations(Context);
+	
+		return UToolMenus::Get()->GenerateWidget("LevelEditor.ActorContextMenu", Context);
+	}
+	
+	if (DoesSelectionHaveType(InData, UActorComponent::StaticClass()))
+	{
+		FLevelEditorContextMenu::RegisterComponentContextMenu();
+		FLevelEditorContextMenu::RegisterElementContextMenu();
+
+		PerformLevelEditorRegistrations(Context);
+
+		return UToolMenus::Get()->GenerateWidget("LevelEditor.ComponentContextMenu", Context);
+	}
+
+	if (DoesSelectionHaveType(InData, UObject::StaticClass()))
+	{
+		FLevelEditorContextMenu::RegisterElementContextMenu();
+
+		PerformLevelEditorRegistrations(Context);
+
+		return UToolMenus::Get()->GenerateWidget("LevelEditor.ElementContextMenu", Context);
+	}
+
+	// Folders only
+	RegisterFoldersOnlyContextMenu();
+
+	return UToolMenus::Get()->GenerateWidget("ObjectMixer.FoldersOnlyContextMenu", Context);
+}
 
-	return UToolMenus::Get()->GenerateWidget("LevelEditor.ActorContextMenu", Context);
+void UObjectMixerEditorListMenuContext::FillSelectionSubMenu(UToolMenu* Menu, const FObjectMixerEditorListMenuContextData& ContextData)
+{
+	FToolMenuSection& Section = Menu->AddSection("Section");
+	Section.AddMenuEntry(
+		"AddChildrenToSelection",
+		LOCTEXT( "AddChildrenToSelection", "Immediate Children" ),
+		LOCTEXT( "AddChildrenToSelection_ToolTip", "Select all immediate children of the selected folders" ),
+		FSlateIcon(),
+		FExecuteAction::CreateStatic(&UObjectMixerEditorListMenuContext::SelectDescendentsOfSelectedFolders, ContextData, false));
+	Section.AddMenuEntry(
+		"AddDescendantsToSelection",
+		LOCTEXT( "AddDescendantsToSelection", "All Descendants" ),
+		LOCTEXT( "AddDescendantsToSelection_ToolTip", "Select all descendants of the selected folders" ),
+		FSlateIcon(),
+		FExecuteAction::CreateStatic(&UObjectMixerEditorListMenuContext::SelectDescendentsOfSelectedFolders, ContextData, true));
 }
 
-void UObjectMixerEditorListMenuContext::RegisterObjectMixerContextMenu()
+void UObjectMixerEditorListMenuContext::RegisterFoldersOnlyContextMenu()
+{
+	UToolMenus* ToolMenus = UToolMenus::Get();
+	if (ToolMenus->IsMenuRegistered("ObjectMixer.FoldersOnlyContextMenu"))
+	{
+		return;
+	}
+
+	UToolMenu* Menu = ToolMenus->RegisterMenu("ObjectMixer.FoldersOnlyContextMenu");
+	Menu->AddDynamicSection("FolderContextMenuDynamic", FNewToolMenuDelegate::CreateLambda([](UToolMenu* InMenu)
+	{
+		// Ensure proper context
+		UObjectMixerEditorListMenuContext* Context = InMenu->FindContext<UObjectMixerEditorListMenuContext>();
+		if (!Context || Context->Data.SelectedItems.Num() == 0)
+		{
+			return;
+		}
+
+		const FObjectMixerEditorListMenuContextData& ContextData = Context->Data;
+
+		{
+			FToolMenuSection& Section = InMenu->AddSection("Hierarchy", LOCTEXT("HierarchyMenuHeader", "Hierarchy"));
+
+			if (ContextData.SelectedItems.Num() == 1)
+			{
+				const FSlateIcon NewFolderIcon(FAppStyle::GetAppStyleSetName(), "SceneOutliner.NewFolderIcon");
+				
+				Section.AddMenuEntry(
+				   "CreateSubFolder",
+				   LOCTEXT("CreateSubFolder", "Create Sub Folder"),
+				   FText(),
+				   NewFolderIcon,
+				   FUIAction(FExecuteAction::CreateLambda([ContextData]()
+				   {
+				   		if (TSharedPtr<FObjectMixerEditorMainPanel> PinnedPanel = ContextData.MainPanelPtr.Pin())
+				   		{
+				   			PinnedPanel->OnRequestNewFolder(ContextData.SelectedItems[0]->GetFolder());
+				   		}
+				   }))
+			   );
+			}
+
+			Section.AddSubMenu(
+				"SelectSubMenu",
+				LOCTEXT("SelectSubMenu", "Select"),
+				FText::GetEmpty(),
+				FNewToolMenuDelegate::CreateLambda(
+					[ContextData](UToolMenu* InMenu)
+				{
+					FillSelectionSubMenu(InMenu, ContextData);
+				})
+			);
+		}
+
+		// For future CL
+		// {
+		// 	FToolMenuSection& Section = InMenu->AddSection("ElementEditActions");
+		// 	Section.AddSubMenu(
+		// 		"EditSubMenu",
+		// 		LOCTEXT("EditSubMenu", "Edit"),
+		// 		FText::GetEmpty(),
+		// 		FNewToolMenuDelegate::CreateLambda([ContextData](UToolMenu* InMenu)
+		// 		{
+		// 			FToolMenuSection& Section = InMenu->AddSection("Section");
+		// 			Section.AddEntry(MakeCustomEditMenu(ContextData));
+		// 		}),
+		// 		false, // default value
+		// 		FSlateIcon(FAppStyle::GetAppStyleSetName(), "Icons.Edit"));
+		// }
+
+		GenerateMoveToMenu(InMenu, FToolMenuInsert("Hierarchy", EToolMenuInsertType::After), ContextData);
+	}));
+}
+
+void UObjectMixerEditorListMenuContext::RegisterObjectMixerActorContextMenuExtension()
 {
 	if (UToolMenu* ActorContextMenu = UToolMenus::Get()->ExtendMenu("LevelEditor.ActorContextMenu"))
 	{		
@@ -135,6 +279,8 @@ void UObjectMixerEditorListMenuContext::RegisterObjectMixerContextMenu()
 						FNewToolMenuDelegate::CreateStatic(&UObjectMixerEditorListMenuContext::CreateSelectCollectionsSubMenu, ContextData)
 					);
 				}
+
+				GenerateMoveToMenu(InMenu, FToolMenuInsert("ActorTypeTools", EToolMenuInsertType::After), ContextData);
 			}),
 			FToolMenuInsert(NAME_None,EToolMenuInsertType::First)
 		);
@@ -195,6 +341,161 @@ void UObjectMixerEditorListMenuContext::CreateSelectCollectionsSubMenu(UToolMenu
 	}
 }
 
+void UObjectMixerEditorListMenuContext::GenerateMoveToMenu(
+	UToolMenu* InMenu, const FToolMenuInsert& InsertArgs, const FObjectMixerEditorListMenuContextData& ContextData)
+{
+	FToolMenuSection& Section =
+		InMenu->AddSection("MainSection",
+			LOCTEXT("OutlinerSectionName", "Outliner"), InsertArgs);
+					
+	Section.AddSubMenu(
+			"MoveActorsTo",
+			LOCTEXT("MoveActorsTo", "Move To"),
+			LOCTEXT("MoveActorsTo_Tooltip", "Move selection to another folder"),
+			FNewToolMenuDelegate::CreateStatic(&UObjectMixerEditorListMenuContext::FillFoldersSubMenu, ContextData));
+}
+
+void UObjectMixerEditorListMenuContext::OnFoldersMenuFolderSelected(
+	TSharedRef<ISceneOutlinerTreeItem> Item,
+	FObjectMixerEditorListMenuContextData ContextData)
+{
+	if (const FFolderTreeItem* FolderTreeItem = Item->CastTo<FFolderTreeItem>())
+	{
+		for (TSharedPtr<FObjectMixerEditorListRow> SelectedItem : ContextData.SelectedItems)
+		{
+			if (AActor* AsActor = Cast<AActor>(SelectedItem->GetObject()))
+			{
+				AsActor->SetFolderPath_Recursively(FolderTreeItem->GetFolder().GetPath());
+			}
+			else if (SelectedItem->GetRowType() == FObjectMixerEditorListRow::Folder)
+			{
+				if (TSharedPtr<FObjectMixerEditorMainPanel> PinnedPanel = ContextData.MainPanelPtr.Pin())
+				{
+					PinnedPanel->OnRequestMoveFolder(SelectedItem->GetFolder(), FolderTreeItem->GetFolder());
+				}
+			}
+		}
+	}
+
+	FSlateApplication::Get().DismissAllMenus();
+}
+
+TSharedRef<TSet<FFolder>> UObjectMixerEditorListMenuContext::GatherInvalidMoveToDestinations(
+	const FObjectMixerEditorListMenuContextData& ContextData)
+{
+	// We use a pointer here to save copying the whole array for every invocation of the filter delegate
+	TSharedRef<TSet<FFolder>> Exclusions(new TSet<FFolder>());
+
+	for (const TSharedPtr<FObjectMixerEditorListRow>& SelectedItem : ContextData.SelectedItems)
+	{
+		if (SelectedItem->GetRowType() != FObjectMixerEditorListRow::Folder)
+		{
+			continue;
+		}
+
+		Exclusions->Add(SelectedItem->GetFolder());
+		
+		if (const TSharedPtr<FObjectMixerEditorListRow>& ParentRow = SelectedItem->GetDirectParentRow().Pin())
+		{
+			auto FolderHasOtherSubFolders = [&SelectedItem](const TSharedPtr<FObjectMixerEditorListRow>& InItem)
+			{
+				if (InItem->GetRowType() == FObjectMixerEditorListRow::Folder && InItem != SelectedItem)
+				{
+					return true;
+				}
+				return false;
+			};
+
+			// Exclude this items direct parent if it is a folder and has no other subfolders we can move to
+			bool bFolderHasOtherSubFolders = false;
+			for (const TSharedPtr<FObjectMixerEditorListRow>& ChildRow : ParentRow->GetChildRows())
+			{
+				if (FolderHasOtherSubFolders(ChildRow))
+				{
+					bFolderHasOtherSubFolders = true;
+					break;
+				}
+			}
+			
+			if (!bFolderHasOtherSubFolders)
+			{
+				Exclusions->Add(ParentRow->GetFolder());
+			}
+		}
+	}
+
+	return Exclusions;
+}
+
+void UObjectMixerEditorListMenuContext::FillFoldersSubMenu(UToolMenu* InMenu, FObjectMixerEditorListMenuContextData ContextData)
+{
+	FToolMenuSection& Section = InMenu->AddSection("Section");
+	Section.AddMenuEntry(
+		"CreateNew",
+		LOCTEXT( "CreateNew", "Create New Folder" ),
+		LOCTEXT( "CreateNew_ToolTip", "Move to a new folder" ),
+		FSlateIcon(FAppStyle::GetAppStyleSetName(), "SceneOutliner.NewFolderIcon"),
+		FExecuteAction::CreateLambda([ContextData]()
+		{
+			if (const TSharedPtr<FObjectMixerEditorMainPanel> PinnedMainPanel = ContextData.MainPanelPtr.Pin())
+			{
+				PinnedMainPanel->OnRequestNewFolder();
+			}
+		})
+	);
+	
+	FSceneOutlinerInitializationOptions MiniSceneOutlinerInitOptions;
+	MiniSceneOutlinerInitOptions.bShowHeaderRow = false;
+	MiniSceneOutlinerInitOptions.bFocusSearchBoxWhenOpened = true;
+
+	MiniSceneOutlinerInitOptions.ModeFactory = FCreateSceneOutlinerMode::CreateLambda(
+		[ContextData](SSceneOutliner* Outliner)
+			{
+				FOnSceneOutlinerItemPicked OnSceneOutlinerItemPicked =
+					FOnSceneOutlinerItemPicked::CreateStatic(
+						&UObjectMixerEditorListMenuContext::OnFoldersMenuFolderSelected, ContextData
+				);
+
+				UWorld* World = FObjectMixerEditorModule::Get().GetWorld();
+			
+				return new FActorFolderPickingMode(Outliner, OnSceneOutlinerItemPicked, World);
+			}
+	);
+
+	TSharedRef<TSet<FFolder>> Exclusions = GatherInvalidMoveToDestinations(ContextData);
+	MiniSceneOutlinerInitOptions.Filters->AddFilterPredicate<FFolderTreeItem>(
+		FFolderTreeItem::FFilterPredicate::CreateLambda(
+			[&Exclusions](const FFolder& Folder)
+			{
+				for (const FFolder& Parent : *Exclusions)
+				{
+					if (Folder == Parent || Folder.IsChildOf(Parent))
+					{
+						return false;
+					}
+				}
+				return true;
+			}), FSceneOutlinerFilter::EDefaultBehaviour::Pass);
+
+	TSharedRef< SWidget > MiniSceneOutliner =
+		SNew(SVerticalBox)
+		+ SVerticalBox::Slot()
+		.MaxHeight(400.0f)
+		[
+			SNew(SSceneOutliner, MiniSceneOutlinerInitOptions)
+			.IsEnabled(FSlateApplication::Get().GetNormalExecutionAttribute())
+		];
+
+	FToolMenuSection& MoveToSection = InMenu->AddSection(FName(), LOCTEXT("ExistingFolders", "Existing:"));
+	
+	MoveToSection.AddEntry(FToolMenuEntry::InitWidget(
+		"MiniSceneOutliner",
+		MiniSceneOutliner,
+		FText::GetEmpty(),
+		false)
+	);
+}
+
 void UObjectMixerEditorListMenuContext::OnTextCommitted(const FText& InText, ETextCommit::Type InCommitType, const FObjectMixerEditorListMenuContextData ContextData)
 {
 	if (InCommitType == ETextCommit::OnEnter)
@@ -203,6 +504,18 @@ void UObjectMixerEditorListMenuContext::OnTextCommitted(const FText& InText, ETe
 	}
 }
 
+void UObjectMixerEditorListMenuContext::SelectDescendentsOfSelectedFolders(
+	FObjectMixerEditorListMenuContextData ContextData, const bool bRecursive)
+{
+	for (const TSharedPtr<FObjectMixerEditorListRow>& SelectedRow : ContextData.SelectedItems)
+	{
+		if (SelectedRow->GetRowType() == FObjectMixerEditorListRow::Folder)
+		{
+			SelectedRow->SetChildRowsSelected(true, bRecursive);
+		}
+	}
+}
+
 void UObjectMixerEditorListMenuContext::OnClickCollectionMenuEntry(const FName Key, const FObjectMixerEditorListMenuContextData ContextData)
 {
 	if (AreAllObjectsInCollection(Key, ContextData))
@@ -274,6 +587,41 @@ bool UObjectMixerEditorListMenuContext::AreAllObjectsInCollection(const FName Ke
 	return bAreAllSelectedObjectsInCollection;
 }
 
+FToolMenuEntry UObjectMixerEditorListMenuContext::MakeCustomEditMenu(const FObjectMixerEditorListMenuContextData& ContextData)
+{
+	FToolMenuEntry Entry;
+	Entry.Name = TEXT("ObjectMixerGenericCommands");
+	Entry.Type = EMultiBlockType::Widget;
+	FToolMenuEntryWidgetData WidgetData;
+	WidgetData.bNoIndent = true;
+	WidgetData.bNoPadding = true;
+	Entry.WidgetData = WidgetData;
+	Entry.MakeCustomWidget.BindLambda(
+		[ContextData](const FToolMenuContext&, const FToolMenuCustomWidgetContext&)
+		{
+			if (TSharedPtr<FObjectMixerEditorMainPanel> PinnedMainPanel = ContextData.MainPanelPtr.Pin())
+			{
+				// Add options with our mapped command list to an FMenuBuilder
+				FMenuBuilder Builder(true, PinnedMainPanel->ObjectMixerElementEditCommands);
+				{
+					Builder.AddMenuEntry( FGenericCommands::Get().Cut);
+					Builder.AddMenuEntry( FGenericCommands::Get().Copy);
+					Builder.AddMenuEntry( FGenericCommands::Get().Paste);
+					Builder.AddMenuEntry( FGenericCommands::Get().Duplicate);
+					Builder.AddMenuEntry( FGenericCommands::Get().Delete);
+					Builder.AddMenuEntry( FGenericCommands::Get().Rename);
+				}
+
+				return Builder.MakeWidget();
+			}
+
+			return SNullWidget::NullWidget;
+		}
+	);
+
+	return Entry;
+}
+
 void UObjectMixerEditorListMenuContext::ReplaceEditSubMenu(const FObjectMixerEditorListMenuContextData& ContextData)
 {
 	if (UToolMenu* EditSubMenu = UToolMenus::Get()->ExtendMenu("LevelEditor.ActorContextMenu.EditSubMenu"))
@@ -291,36 +639,7 @@ void UObjectMixerEditorListMenuContext::ReplaceEditSubMenu(const FObjectMixerEdi
 
 		// Add our menu widget to the edit submenu
 		// It must be a widget entry because we hid the entries above, so they won't show as regular menu entries since they share names
-		FToolMenuEntry Entry;
-		Entry.Name = TEXT("ObjectMixerGenericCommands");
-		Entry.Type = EMultiBlockType::Widget;
-		FToolMenuEntryWidgetData WidgetData;
-		WidgetData.bNoIndent = true;
-		WidgetData.bNoPadding = true;
-		Entry.WidgetData = WidgetData;
-		Entry.MakeCustomWidget.BindLambda(
-			[ContextData](const FToolMenuContext&, const FToolMenuCustomWidgetContext&)
-			{
-				if (TSharedPtr<FObjectMixerEditorMainPanel> PinnedMainPanel = ContextData.MainPanelPtr.Pin())
-				{
-					// Add options with our mapped command list to an FMenuBuilder
-					FMenuBuilder Builder(true, PinnedMainPanel->ObjectMixerCommands);
-					{
-						Builder.AddMenuEntry( FGenericCommands::Get().Cut);
-						Builder.AddMenuEntry( FGenericCommands::Get().Copy);
-						Builder.AddMenuEntry( FGenericCommands::Get().Paste);
-						Builder.AddMenuEntry( FGenericCommands::Get().Duplicate);
-						Builder.AddMenuEntry( FGenericCommands::Get().Delete);
-						Builder.AddMenuEntry( FGenericCommands::Get().Rename);
-					}
-
-					return Builder.MakeWidget();
-				}
-
-				return SNullWidget::NullWidget;
-			}
-		);
-		EditSubMenu->AddMenuEntry(NAME_None, Entry);
+		EditSubMenu->AddMenuEntry(NAME_None, MakeCustomEditMenu(ContextData));
 	}
 }
 
diff --git a/Engine/Plugins/Editor/ObjectMixer/ObjectMixer/Source/ObjectMixer/Public/ObjectMixerEditorModule.h b/Engine/Plugins/Editor/ObjectMixer/ObjectMixer/Source/ObjectMixer/Public/ObjectMixerEditorModule.h
--- a/Engine/Plugins/Editor/ObjectMixer/ObjectMixer/Source/ObjectMixer/Public/ObjectMixerEditorModule.h
+++ b/Engine/Plugins/Editor/ObjectMixer/ObjectMixer/Source/ObjectMixer/Public/ObjectMixerEditorModule.h
@@ -17,6 +17,8 @@ public:
 	virtual void ShutdownModule() override;
 	//~ End IModuleInterface
 
+	virtual UWorld* GetWorld();
+
 	virtual void Initialize();
 	virtual void Teardown();
 
diff --git a/Engine/Plugins/Editor/ObjectMixer/ObjectMixer/Source/ObjectMixer/Public/Views/List/ObjectMixerEditorList.h b/Engine/Plugins/Editor/ObjectMixer/ObjectMixer/Source/ObjectMixer/Public/Views/List/ObjectMixerEditorList.h
--- a/Engine/Plugins/Editor/ObjectMixer/ObjectMixer/Source/ObjectMixer/Public/Views/List/ObjectMixerEditorList.h
+++ b/Engine/Plugins/Editor/ObjectMixer/ObjectMixer/Source/ObjectMixer/Public/Views/List/ObjectMixerEditorList.h
@@ -9,6 +9,8 @@
 
 class SObjectMixerEditorList;
 
+struct FFolder;
+
 class OBJECTMIXEREDITOR_API FObjectMixerEditorList : public TSharedFromThis<FObjectMixerEditorList>
 {
 public:
@@ -42,6 +44,10 @@ public:
 
 	void RequestSyncEditorSelectionToListSelection() const;
 
+	void OnRequestNewFolder(TOptional<FFolder> ExplicitParentFolder = TOptional<FFolder>());
+
+	void OnRequestMoveFolder(const FFolder& FolderToMove, const FFolder& TargetNewParentFolder);
+
 	void ExecuteListViewSearchOnAllRows(const FString& SearchString, const bool bShouldRefreshAfterward = true);
 
 	void EvaluateIfRowsPassFilters(const bool bShouldRefreshAfterward = true) const;
diff --git a/Engine/Plugins/Editor/ObjectMixer/ObjectMixer/Source/ObjectMixer/Public/Views/List/ObjectMixerEditorListRow.h b/Engine/Plugins/Editor/ObjectMixer/ObjectMixer/Source/ObjectMixer/Public/Views/List/ObjectMixerEditorListRow.h
--- a/Engine/Plugins/Editor/ObjectMixer/ObjectMixer/Source/ObjectMixer/Public/Views/List/ObjectMixerEditorListRow.h
+++ b/Engine/Plugins/Editor/ObjectMixer/ObjectMixer/Source/ObjectMixer/Public/Views/List/ObjectMixerEditorListRow.h
@@ -8,6 +8,8 @@
 #include "Containers/UnrealString.h"
 #include "DragAndDrop/DecoratedDragDropOp.h"
 #include "Folder.h"
+#include "GameFramework/Actor.h"
+#include "ISceneOutlinerTreeItem.h"
 #include "Layout/Visibility.h"
 #include "PropertyHandle.h"
 #include "Templates/SharedPointer.h"
@@ -110,6 +112,23 @@ struct OBJECTMIXEREDITOR_API FObjectMixerEditorListRow final : TSharedFromThis<F
 		return nullptr;
 	}
 
+	[[nodiscard]] AActor* GetSelfOrOuterAsActor() const
+	{
+		if (UObject* Object = GetObject())
+		{
+			AActor* Actor = Cast<AActor>(Object);
+
+			if (!Actor)
+			{
+				Actor = Object->GetTypedOuter<AActor>();
+			}
+
+			return Actor;
+		}
+
+		return nullptr;
+	}
+
 	[[nodiscard]] const FFolder& GetFolder() const
 	{
 		return FolderRef;
@@ -146,6 +165,8 @@ struct OBJECTMIXEREDITOR_API FObjectMixerEditorListRow final : TSharedFromThis<F
 	void AddToChildRows(const FObjectMixerEditorListRowPtr& InRow);
 	void InsertChildRowAtIndex(const FObjectMixerEditorListRowPtr& InRow, const int32 AtIndex = 0);
 
+	void SetChildRowsSelected(const bool bNewSelected, const bool bRecursive, const bool bSelectOnlyVisible = true);
+
 	[[nodiscard]] bool GetIsTreeViewItemExpanded();
 	void SetIsTreeViewItemExpanded(const bool bNewExpanded);
 
@@ -167,6 +188,7 @@ struct OBJECTMIXEREDITOR_API FObjectMixerEditorListRow final : TSharedFromThis<F
 	void SetDoesRowPassFilters(const bool bPass);
 
 	[[nodiscard]] bool GetIsSelected();
+	void SetIsSelected(const bool bNewSelected);
 
 	[[nodiscard]] bool ShouldRowWidgetBeVisible() const;
 	[[nodiscard]] EVisibility GetDesiredRowWidgetVisibility() const;
@@ -222,7 +244,7 @@ struct OBJECTMIXEREDITOR_API FObjectMixerEditorListRow final : TSharedFromThis<F
 
 	void CallOnRenameCommandDelegate()
 	{
-		OnRenameCommandDelegate.Execute();
+		OnRenameCommandDelegate.ExecuteIfBound();
 	}
 	
 	TMap<FName, TWeakPtr<IPropertyHandle>> PropertyNamesToHandles;
@@ -230,7 +252,7 @@ struct OBJECTMIXEREDITOR_API FObjectMixerEditorListRow final : TSharedFromThis<F
 private:
 
 	FObjectMixerEditorListRowPtr GetAsShared();
-
+	
 	TWeakObjectPtr<UObject> ObjectRef;
 	FFolder FolderRef;
 	EObjectMixerEditorListRowType RowType = MatchingObject;
diff --git a/Engine/Plugins/Editor/ObjectMixer/ObjectMixer/Source/ObjectMixer/Public/Views/List/SObjectMixerEditorList.h b/Engine/Plugins/Editor/ObjectMixer/ObjectMixer/Source/ObjectMixer/Public/Views/List/SObjectMixerEditorList.h
--- a/Engine/Plugins/Editor/ObjectMixer/ObjectMixer/Source/ObjectMixer/Public/Views/List/SObjectMixerEditorList.h
+++ b/Engine/Plugins/Editor/ObjectMixer/ObjectMixer/Source/ObjectMixer/Public/Views/List/SObjectMixerEditorList.h
@@ -77,6 +77,8 @@ public:
 
 	virtual ~SObjectMixerEditorList() override;
 
+	UWorld* GetWorld() const;
+
 	TWeakPtr<FObjectMixerEditorList> GetListModelPtr()
 	{
 		return ListModelPtr;
@@ -103,6 +105,8 @@ public:
 
 	[[nodiscard]] TArray<FObjectMixerEditorListRowPtr> GetSelectedTreeViewItems() const;
 	int32 GetSelectedTreeViewItemCount() const;
+
+	void SetTreeViewItemSelected(FObjectMixerEditorListRowPtr Item, const bool bNewSelected);
 	
 	void RequestSyncEditorSelectionToListSelection()
 	{
@@ -110,6 +114,10 @@ public:
 	}
 	void SyncEditorSelectionToListSelection();
 
+	void OnRequestNewFolder(TOptional<FFolder> ExplicitParentFolder = TOptional<FFolder>());
+
+	void OnRequestMoveFolder(const FFolder& FolderToMove, const FFolder& TargetNewParentFolder);
+
 	void SetSelectedTreeViewItemActorsEditorVisible(const bool bNewIsVisible, const bool bIsRecursive = false);
 
 	bool IsTreeViewItemSelected(TSharedRef<FObjectMixerEditorListRow> Item);
@@ -190,6 +198,7 @@ public:
 
 protected:
 
+	virtual void ExecuteRenameOnPending();
 	virtual void Tick(const FGeometry& AllottedGeometry, const double InCurrentTime, const float InDeltaTime) override;
 
 	void OnActorSpawnedOrDestroyed(AActor* Object)
@@ -247,10 +256,49 @@ protected:
 
 	TSharedPtr<STreeView<FObjectMixerEditorListRowPtr>> TreeViewPtr;
 
-	struct FTreeItemStateCache
+	struct FTreeItemUniqueIdentifier
 	{
-		uint32 UniqueId = -1;
+		FTreeItemUniqueIdentifier() = default;
+		
+		FTreeItemUniqueIdentifier(FObjectMixerEditorListRowPtr InTreeViewItem)
+		{
+			const UObject* RowObject = InTreeViewItem->GetObject();
+			
+			UniqueIdNumber = RowObject ? RowObject->GetUniqueID() : -1;
+			
+			RowName =
+				InTreeViewItem->GetRowType() == FObjectMixerEditorListRow::Folder ?
+					InTreeViewItem->GetFolderPath().ToString() : InTreeViewItem->GetDisplayName().ToString();
+		}
+
+		bool operator==(const FTreeItemUniqueIdentifier& Other) const
+		{
+			if (Other.UniqueIdNumber != -1)
+			{
+				return Other.UniqueIdNumber == UniqueIdNumber;
+			}
+
+			return Other.RowName.Equals(RowName);
+		}
+
+		bool IdentifiesAnyRow() const
+		{
+			return UniqueIdNumber != 1 || !RowName.IsEmpty();
+		}
+
+		void Reset()
+		{
+			UniqueIdNumber = -1;
+			RowName = "";
+		}
+		
+		uint32 UniqueIdNumber = -1;
 		FString RowName = "";
+	};
+
+	struct FTreeItemStateCache
+	{
+		FTreeItemUniqueIdentifier UniqueIdentifier;
 		bool bIsExpanded = false;
 		bool bIsSelected = false;
 		FObjectMixerEditorListRow::FTransientEditorVisibilityRules VisibilityRules;
@@ -264,6 +312,8 @@ protected:
 
 	TArray<FFilterComboToStateCaches> FilterComboToStateCaches;
 	
+	FTreeItemUniqueIdentifier PendingRenameItem;
+	
 	TMap<UObject*, FObjectMixerEditorListRowPtr> ObjectsToRowsCreated;
 
 	/** All Tree view objects */
diff --git a/Engine/Plugins/Editor/ObjectMixer/ObjectMixer/Source/ObjectMixer/Public/Views/MainPanel/ObjectMixerEditorMainPanel.h b/Engine/Plugins/Editor/ObjectMixer/ObjectMixer/Source/ObjectMixer/Public/Views/MainPanel/ObjectMixerEditorMainPanel.h
--- a/Engine/Plugins/Editor/ObjectMixer/ObjectMixer/Source/ObjectMixer/Public/Views/MainPanel/ObjectMixerEditorMainPanel.h
+++ b/Engine/Plugins/Editor/ObjectMixer/ObjectMixer/Source/ObjectMixer/Public/Views/MainPanel/ObjectMixerEditorMainPanel.h
@@ -4,17 +4,19 @@
 
 #include "ObjectFilter/ObjectMixerEditorObjectFilter.h"
 
+#include "Folder.h"
 #include "UObject/StrongObjectPtr.h"
 #include "Widgets/SWidget.h"
 
-class FUICommandList;
+class FObjectMixerEditorList;
 class FObjectMixerEditorListFilter_Collection;
-class UObjectMixerEditorSerializedData;
-struct FObjectMixerEditorListRow;
+class FUICommandList;
 class IObjectMixerEditorListFilter;
-class UObjectMixerObjectFilter;
-class FObjectMixerEditorList;
 class SObjectMixerEditorMainPanel;
+class UObjectMixerEditorSerializedData;
+class UObjectMixerObjectFilter;
+
+struct FObjectMixerEditorListRow;
 
 DECLARE_MULTICAST_DELEGATE(FOnPreFilterChange)
 DECLARE_MULTICAST_DELEGATE(FOnPostFilterChange)
@@ -50,6 +52,10 @@ public:
 	/** Called when the Rename command is executed from the UI or hotkey. */
 	void OnRenameCommand();
 
+	void OnRequestNewFolder(TOptional<FFolder> ExplicitParentFolder = TOptional<FFolder>());
+
+	void OnRequestMoveFolder(const FFolder& FolderToMove, const FFolder& TargetNewParentFolder);
+
 	void RequestSyncEditorSelectionToListSelection() const;
 
 	TWeakPtr<FObjectMixerEditorList> GetEditorListModel() const
@@ -162,7 +168,8 @@ public:
 	FOnPreFilterChange OnPreFilterChange;
 	FOnPostFilterChange OnPostFilterChange;
 
-	TSharedPtr<FUICommandList> ObjectMixerCommands;
+	TSharedPtr<FUICommandList> ObjectMixerElementEditCommands;
+	TSharedPtr<FUICommandList> ObjectMixerFolderEditCommands;
 
 private:
 
diff --git a/Engine/Plugins/Editor/ObjectMixer/ObjectMixer/Source/ObjectMixer/Public/Views/Widgets/ObjectMixerEditorListMenuContext.h b/Engine/Plugins/Editor/ObjectMixer/ObjectMixer/Source/ObjectMixer/Public/Views/Widgets/ObjectMixerEditorListMenuContext.h
--- a/Engine/Plugins/Editor/ObjectMixer/ObjectMixer/Source/ObjectMixer/Public/Views/Widgets/ObjectMixerEditorListMenuContext.h
+++ b/Engine/Plugins/Editor/ObjectMixer/ObjectMixer/Source/ObjectMixer/Public/Views/Widgets/ObjectMixerEditorListMenuContext.h
@@ -4,6 +4,7 @@
 
 #include "Views/List/ObjectMixerEditorListRow.h"
 
+#include "ToolMenuEntry.h"
 #include "Types/SlateEnums.h"
 
 #include "ObjectMixerEditorListMenuContext.generated.h"	
@@ -26,7 +27,8 @@ public:
 
 	static TSharedPtr<SWidget> CreateContextMenu(const FObjectMixerEditorListMenuContextData InData);
 	static TSharedPtr<SWidget> BuildContextMenu(const FObjectMixerEditorListMenuContextData& InData);
-	static void RegisterObjectMixerContextMenu();
+	static void RegisterFoldersOnlyContextMenu();
+	static void RegisterObjectMixerActorContextMenuExtension();
 
 	FObjectMixerEditorListMenuContextData Data;
 	
@@ -34,13 +36,22 @@ public:
 
 private:
 
+	static bool DoesSelectionHaveType(const FObjectMixerEditorListMenuContextData& InData, UClass* Type);
 	static void CreateSelectCollectionsSubMenu(UToolMenu* Menu, FObjectMixerEditorListMenuContextData ContextData);
+	static void GenerateMoveToMenu(UToolMenu* InMenu, const FToolMenuInsert& InsertArgs, const FObjectMixerEditorListMenuContextData& ContextData);
+	static void OnFoldersMenuFolderSelected(TSharedRef<ISceneOutlinerTreeItem> Item, FObjectMixerEditorListMenuContextData ContextData);
+	static TSharedRef<TSet<FFolder>> GatherInvalidMoveToDestinations(const FObjectMixerEditorListMenuContextData& ContextData);
+	static void FillFoldersSubMenu(UToolMenu* InMenu, FObjectMixerEditorListMenuContextData ContextData);
+	static void FillSelectionSubMenu(UToolMenu* Menu, const FObjectMixerEditorListMenuContextData& ContextData);
 	static void OnTextCommitted(const FText& InText, ETextCommit::Type InCommitType, const FObjectMixerEditorListMenuContextData ContextData);
+
+	static void SelectDescendentsOfSelectedFolders(FObjectMixerEditorListMenuContextData ContextData, const bool bRecursive);
 	
 	static void OnClickCollectionMenuEntry(const FName Key, const FObjectMixerEditorListMenuContextData ContextData);
 	static void AddObjectsToCollection(const FName Key, const FObjectMixerEditorListMenuContextData ContextData);
 	static void RemoveObjectsFromCollection(const FName Key, const FObjectMixerEditorListMenuContextData ContextData);
 	static bool AreAllObjectsInCollection(const FName Key, const FObjectMixerEditorListMenuContextData ContextData);
 
+	static FToolMenuEntry MakeCustomEditMenu(const FObjectMixerEditorListMenuContextData& ContextData);
 	static void ReplaceEditSubMenu(const FObjectMixerEditorListMenuContextData& ContextData);
 };
