Object Mixer:

-Fix: Reparenting/Moving A Child Light to Folder Bug
-Fix: Folder drag and drop into same parent folder results in name change
-Fix: Create new folder from selection button doesn't include selected folder rows
-Fix: Hybrid Row blueprints have no edit link

diff --git a/Engine/Plugins/Editor/ObjectMixer/ObjectMixer/Source/ObjectMixer/Private/Views/List/SObjectMixerEditorList.cpp b/Engine/Plugins/Editor/ObjectMixer/ObjectMixer/Source/ObjectMixer/Private/Views/List/SObjectMixerEditorList.cpp
--- a/Engine/Plugins/Editor/ObjectMixer/ObjectMixer/Source/ObjectMixer/Private/Views/List/SObjectMixerEditorList.cpp
+++ b/Engine/Plugins/Editor/ObjectMixer/ObjectMixer/Source/ObjectMixer/Private/Views/List/SObjectMixerEditorList.cpp
@@ -265,41 +265,127 @@ void SObjectMixerEditorList::OnRequestNewFolder(TOptional<FFolder> ExplicitParen
 
 	UWorld* World = GetWorld();
 	FFolder NewFolderPath;
+	const TArray<FObjectMixerEditorListRowPtr>& SelectedTreeItems = TreeViewPtr->GetSelectedItems();
 	if (ExplicitParentFolder.IsSet())
 	{
 		NewFolderPath =
 			FActorFolders::Get().GetDefaultFolderName(*World, ExplicitParentFolder.GetValue());
 	}
 	else
+	{
+		// Explicit folder not provided - try to find an appropriate parent folder based on selection
+		if (SelectedTreeItems.Num())
+		{
+			FFolder FoundParent;
+			int32 FolderCount = 0;
+
+			// Check for a folder row first. It's possible for a root folder to be selected.
+			// If that's the case, it won't have a parent folder but we won't want to move it to a selected actor's direct folder.
+			// For this reason, we only check selected actors if there are no folders selected or there is no folder match
+			const TSharedPtr<FObjectMixerEditorListRow>* FolderMatch = Algo::FindByPredicate(
+				SelectedTreeItems,
+				[&FoundParent, &FolderCount](const FObjectMixerEditorListRowPtr& Other)
+				{
+					if (Other->GetRowType() == FObjectMixerEditorListRow::Folder)
+					{
+						FolderCount++;
+						FoundParent = Other->GetFolder().GetParent();
+					}
+
+					return !FoundParent.IsNone();
+				}
+			);
+
+			// If no folder row found, try actors
+			if (FolderCount == 0 || !FolderMatch)
+			{
+				const TSharedPtr<FObjectMixerEditorListRow>* ActorMatch = Algo::FindByPredicate(
+					SelectedTreeItems,
+					[&FoundParent](const FObjectMixerEditorListRowPtr& Other)
+					{
+						if (const AActor* Actor = Other->GetSelfOrOuterAsActor())
+						{
+							FoundParent = Actor->GetFolder();
+						}
+
+						return !FoundParent.IsNone();
+					}
+				);
+			}
+
+			if (!FoundParent.IsNone())
+			{
+				NewFolderPath =
+				   FActorFolders::Get().GetDefaultFolderName(*World, FoundParent);
+			}
+		}
+	}
+
+	// Fallback
+	if (NewFolderPath.IsNone())
 	{
 		NewFolderPath = FActorFolders::Get().GetDefaultFolderForSelection(*World);
 	}
 
 	if (FActorFolders::Get().CreateFolder(*World, NewFolderPath))
 	{
-		const TArray<FObjectMixerEditorListRowPtr>& SelectedTreeItems = TreeViewPtr->GetSelectedItems();
-
 		if (SelectedTreeItems.Num())
 		{
-			TArray<AActor*> ActorsToMoveToFolder;
+			TArray<AActor*> ActorsToMove;
+			TSet<FFolder> FoldersToMove;
 			for (const FObjectMixerEditorListRowPtr& SelectedRow : SelectedTreeItems)
 			{
 				if (SelectedRow->GetRowType() == FObjectMixerEditorListRow::MatchingObject ||
 					SelectedRow->GetRowType() == FObjectMixerEditorListRow::ContainerObject)
 				{
-					if (AActor* Actor = SelectedRow->GetSelfOrOuterAsActor())
+					if (FObjectMixerEditorListRowPtr PinnedParent = SelectedRow->GetDirectParentRow().Pin();
+						PinnedParent->GetRowType() == FObjectMixerEditorListRow::Folder &&
+						SelectedTreeItems.Contains(PinnedParent))
+					{
+						// If the selected item's parent is a folder and it's selected, just move the folder.
+						// It will take the selected objects along with it.
+						const FFolder& ParentFolder = SelectedRow->GetFolder();
+
+						if (!ParentFolder.IsNone())
+						{
+							FoldersToMove.Add(ParentFolder);
+						}
+					}
+					else if (AActor* Actor = SelectedRow->GetSelfOrOuterAsActor())
 					{
-						ActorsToMoveToFolder.Add(Actor);
+						ActorsToMove.Add(Actor);
 					}
 				}
+				else if (SelectedRow->GetRowType() == FObjectMixerEditorListRow::Folder)
+				{
+					FoldersToMove.Add(SelectedRow->GetFolder());
+				}
 			}
 
-			if (ActorsToMoveToFolder.Num())
+			if (FoldersToMove.Num())
+			{
+				for (const FFolder& Folder : FoldersToMove)
+				{
+					OnRequestMoveFolder(Folder, NewFolderPath);
+				}
+			}
+
+			if (ActorsToMove.Num())
 			{				
-				for (AActor* Actor : ActorsToMoveToFolder)
+				for (AActor* Actor : ActorsToMove)
 				{
 					Actor->Modify();
 					Actor->SetFolderPath_Recursively(NewFolderPath.GetPath());
+
+					if (AActor* AttachParent = Actor->GetAttachParentActor())
+					{
+						if (!ActorsToMove.Contains(AttachParent))
+						{
+							AttachParent->Modify();
+							FDetachmentTransformRules DetachmentRules(EDetachmentRule::KeepWorld, false);
+							Actor->DetachFromActor(DetachmentRules);
+						}
+					}
 				}
 			}
 		}
@@ -313,10 +399,13 @@ void SObjectMixerEditorList::OnRequestMoveFolder(const FFolder& FolderToMove, co
 {
 	if (UWorld* EditorWorld = GetWorld())
 	{
-		check(TargetNewParentFolder.GetRootObject() == FolderToMove.GetRootObject());
-		// Get unique name
-		const FFolder NewPath = FActorFolders::Get().GetFolderName(*EditorWorld, TargetNewParentFolder, FolderToMove.GetLeafName());
-		FActorFolders::Get().RenameFolderInWorld(*EditorWorld, FolderToMove, NewPath);
+		if (FolderToMove.GetParent() != TargetNewParentFolder)
+		{
+			check(TargetNewParentFolder.GetRootObject() == FolderToMove.GetRootObject());
+			// Get unique name
+			const FFolder NewPath = FActorFolders::Get().GetFolderName(*EditorWorld, TargetNewParentFolder, FolderToMove.GetLeafName());
+			FActorFolders::Get().RenameFolderInWorld(*EditorWorld, FolderToMove, NewPath);
+		}
 	}
 }
 
diff --git a/Engine/Plugins/Editor/ObjectMixer/ObjectMixer/Source/ObjectMixer/Private/Views/List/SObjectMixerEditorListRow.cpp b/Engine/Plugins/Editor/ObjectMixer/ObjectMixer/Source/ObjectMixer/Private/Views/List/SObjectMixerEditorListRow.cpp
--- a/Engine/Plugins/Editor/ObjectMixer/ObjectMixer/Source/ObjectMixer/Private/Views/List/SObjectMixerEditorListRow.cpp
+++ b/Engine/Plugins/Editor/ObjectMixer/ObjectMixer/Source/ObjectMixer/Private/Views/List/SObjectMixerEditorListRow.cpp
@@ -614,7 +614,7 @@ FReply SObjectMixerEditorListRow::HandleAcceptDrop(const FDragDropEvent& DragDro
 			{
 				if (const TSharedPtr<SObjectMixerEditorList> PinnedList = DraggedItem->GetListViewPtr().Pin())
 				{
-					PinnedList->OnRequestMoveFolder(DraggedItem->GetFolder(), TargetItem->GetFolder());
+						PinnedList->OnRequestMoveFolder(DraggedItem->GetFolder(), TargetItem->GetFolder());
 				}
 			}
 		}
@@ -624,6 +624,19 @@ FReply SObjectMixerEditorListRow::HandleAcceptDrop(const FDragDropEvent& DragDro
 			{
 				ObjectAsActor->Modify();
 				ObjectAsActor->SetFolderPath(TargetItem->GetFolderPath());
+
+				if (AActor* AttachParent = ObjectAsActor->GetAttachParentActor())
+				{
+					if (const FObjectMixerEditorListRowPtr ParentRow = DraggedItem->GetDirectParentRow().Pin(); ParentRow.IsValid())
+					{
+						if (!Operation->DraggedItems.Contains(ParentRow))
+						{
+							AttachParent->Modify();
+							FDetachmentTransformRules DetachmentRules(EDetachmentRule::KeepWorld, false);
+							ObjectAsActor->DetachFromActor(DetachmentRules);
+						}
+					}
+				}
 			}
 			else if (AActor* DropOnObjectAsActor = Cast<AActor>(DropOnObject))
 			{
