Object Mixer:

-Remove final keyword from classes

diff --git a/Engine/Plugins/Editor/ObjectMixer/ObjectMixer/Source/ObjectMixer/Public/Views/List/ObjectMixerEditorListRow.h b/Engine/Plugins/Editor/ObjectMixer/ObjectMixer/Source/ObjectMixer/Public/Views/List/ObjectMixerEditorListRow.h
--- a/Engine/Plugins/Editor/ObjectMixer/ObjectMixer/Source/ObjectMixer/Public/Views/List/ObjectMixerEditorListRow.h
+++ b/Engine/Plugins/Editor/ObjectMixer/ObjectMixer/Source/ObjectMixer/Public/Views/List/ObjectMixerEditorListRow.h
@@ -37,7 +37,7 @@ public:
 	static TSharedRef<FObjectMixerListRowDragDropOp> New(const TArray<FObjectMixerEditorListRowPtr>& InItems);
 };
 
-struct OBJECTMIXEREDITOR_API FObjectMixerEditorListRow final : TSharedFromThis<FObjectMixerEditorListRow>
+struct OBJECTMIXEREDITOR_API FObjectMixerEditorListRow : TSharedFromThis<FObjectMixerEditorListRow>
 {
 	enum EObjectMixerEditorListRowType : uint8
 	{
@@ -46,7 +46,7 @@ struct OBJECTMIXEREDITOR_API FObjectMixerEditorListRow final : TSharedFromThis<F
 		ContainerObject, // Usually an actor that contains a matching subobject or is the attach parent of a matching actor
 		MatchingContainerObject, // An Actor that is a matching object and contains matching subobjects 
 		MatchingObject // The object that has the properties we wish to affect
-	};
+	};	
 
 	struct FTransientEditorVisibilityRules
 	{
diff --git a/Engine/Plugins/Editor/ObjectMixer/ObjectMixer/Source/ObjectMixer/Public/Views/List/SObjectMixerEditorList.h b/Engine/Plugins/Editor/ObjectMixer/ObjectMixer/Source/ObjectMixer/Public/Views/List/SObjectMixerEditorList.h
--- a/Engine/Plugins/Editor/ObjectMixer/ObjectMixer/Source/ObjectMixer/Public/Views/List/SObjectMixerEditorList.h
+++ b/Engine/Plugins/Editor/ObjectMixer/ObjectMixer/Source/ObjectMixer/Public/Views/List/SObjectMixerEditorList.h
@@ -63,7 +63,7 @@ struct FListViewColumnInfo
 	float FillWidth = 1.0f;
 };
 
-class OBJECTMIXEREDITOR_API SObjectMixerEditorList final : public SCompoundWidget
+class OBJECTMIXEREDITOR_API SObjectMixerEditorList : public SCompoundWidget
 {
 
 public:
