Object Mixer:

-Fix: Crash when trying to rename blueprint actor row
-Fix: Mutli-selection and Issues with Color Picker Panel
-Fix: Multi-Select Blueprint/Actor w/Lights and Light Parameter Change Bug
-Fix: Shift Multi-Selection Bug
-Set checkbox state as indeterminate when setting collection for multiple objects but not all objects are in given collection

diff --git a/Engine/Plugins/Editor/ObjectMixer/ObjectMixer/Source/ObjectMixer/Private/ObjectMixerEditorModule.cpp b/Engine/Plugins/Editor/ObjectMixer/ObjectMixer/Source/ObjectMixer/Private/ObjectMixerEditorModule.cpp
--- a/Engine/Plugins/Editor/ObjectMixer/ObjectMixer/Source/ObjectMixer/Private/ObjectMixerEditorModule.cpp
+++ b/Engine/Plugins/Editor/ObjectMixer/ObjectMixer/Source/ObjectMixer/Private/ObjectMixerEditorModule.cpp
@@ -353,7 +353,7 @@ void FObjectMixerEditorModule::BindDelegates()
 
 	DelegateHandles.Add(FCoreUObjectDelegates::OnObjectTransacted.AddLambda([this](UObject*, const FTransactionObjectEvent& Event)
 	{
-		if (Event.GetEventType() == ETransactionObjectEventType::Finalized && Event.HasPropertyChanges())
+		if (Event.GetEventType() == ETransactionObjectEventType::Finalized && Event.HasNameChange())
 		{
 			RequestRebuildList();
 		}
diff --git a/Engine/Plugins/Editor/ObjectMixer/ObjectMixer/Source/ObjectMixer/Private/Views/List/ObjectMixerEditorListRow.cpp b/Engine/Plugins/Editor/ObjectMixer/ObjectMixer/Source/ObjectMixer/Private/Views/List/ObjectMixerEditorListRow.cpp
--- a/Engine/Plugins/Editor/ObjectMixer/ObjectMixer/Source/ObjectMixer/Private/Views/List/ObjectMixerEditorListRow.cpp
+++ b/Engine/Plugins/Editor/ObjectMixer/ObjectMixer/Source/ObjectMixer/Private/Views/List/ObjectMixerEditorListRow.cpp
@@ -114,15 +114,18 @@ int32 FObjectMixerEditorListRow::GetOrFindHybridRowIndex()
 		
 				if (ChildRow->GetRowType() == FObjectMixerEditorListRow::MatchingObject)
 				{
-					if (const UObject* ChildObject = ChildRow->GetObject(); ChildObject->GetOuter() == ThisObject)
+					if (const UObject* ChildObject = ChildRow->GetObject())
 					{
-						HybridCandidates.Add(ChildItr);
-
-						if (HybridCandidates.Num() > 1)
+						if (ChildObject->GetOuter() == ThisObject)
 						{
-							// There can only be one row to hybrid with.
-							// If there's more than one candidate, don't hybrid.
-							break;
+							HybridCandidates.Add(ChildItr);
+
+							if (HybridCandidates.Num() > 1)
+							{
+								// There can only be one row to hybrid with.
+								// If there's more than one candidate, don't hybrid.
+								break;
+							}
 						}
 					}
 				}
diff --git a/Engine/Plugins/Editor/ObjectMixer/ObjectMixer/Source/ObjectMixer/Private/Views/List/SObjectMixerEditorListRow.cpp b/Engine/Plugins/Editor/ObjectMixer/ObjectMixer/Source/ObjectMixer/Private/Views/List/SObjectMixerEditorListRow.cpp
--- a/Engine/Plugins/Editor/ObjectMixer/ObjectMixer/Source/ObjectMixer/Private/Views/List/SObjectMixerEditorListRow.cpp
+++ b/Engine/Plugins/Editor/ObjectMixer/ObjectMixer/Source/ObjectMixer/Private/Views/List/SObjectMixerEditorListRow.cpp
@@ -27,7 +27,6 @@
 #include "Widgets/Layout/SBox.h"
 #include "Widgets/SBoxPanel.h"
 #include "Widgets/Text/SInlineEditableTextBlock.h"
-#include "Widgets/Text/STextBlock.h"
 
 #define LOCTEXT_NAMESPACE "ObjectMixerEditor"
 
@@ -66,7 +65,7 @@ public:
 		const bool bRowObjectIsValid = IsValid(RowObject);
 
 		if (bRowObjectIsValid)
-		{
+		{	
 			if (const UClass* ActorClass = RowObject->GetClass())
 			{
 				if (UBlueprint* AsBlueprint = UBlueprint::GetBlueprintFromClass(ActorClass))
@@ -76,13 +75,14 @@ public:
 					HBox->AddSlot()
 					.Padding(FMargin(10.0, 0, 0, 0))
 					[
-						SNew(SHyperlinkWithTextHighlight)
-						.Style(FAppStyle::Get(), "Common.GotoBlueprintHyperlink")
+						SAssignNew(HyperlinkTextBlock, SHyperlinkWithTextHighlight)
 						.Visibility(EVisibility::Visible)
 						.Text(DisplayName)
 						.ToolTipText(LOCTEXT("ClickToEditBlueprint", "Click to edit Blueprint"))
 						.OnNavigate(this, &SInlineEditableRowNameCellWidget::OnClickBlueprintLink, AsBlueprint, RowObject.Get())
 						.HighlightText(this, &SInlineEditableRowNameCellWidget::GetHighlightText)
+						.IsSelected_Raw(this, &SInlineEditableRowNameCellWidget::GetIsSelectedExclusively)
+						.OnTextCommitted(this, &SInlineEditableRowNameCellWidget::OnTextCommitted)
 					];
 				}
 			}
@@ -143,7 +143,14 @@ public:
 
 	void EnterEditingMode() const
 	{
-		EditableTextBlock->EnterEditingMode();
+		if (EditableTextBlock.IsValid())
+		{
+			EditableTextBlock->EnterEditingMode();
+		}
+		else if (HyperlinkTextBlock.IsValid() && HyperlinkTextBlock->EditableTextBlock.IsValid())
+		{
+			HyperlinkTextBlock->EditableTextBlock->EnterEditingMode();
+		}
 	}
 
 private:
@@ -411,6 +418,7 @@ private:
 	TWeakPtr<FObjectMixerEditorListRow> HybridChild;
 
 	TSharedPtr<SInlineEditableTextBlock> EditableTextBlock;
+	TSharedPtr<SHyperlinkWithTextHighlight> HyperlinkTextBlock;
 
 	/** The offset applied to text widgets so that the text aligns with the column header text */
 	float TextBlockLeftPadding = 3.0f;
@@ -945,13 +953,13 @@ TSharedPtr<SWidget> SObjectMixerEditorListRow::GenerateCells(
 						RowPtr->PropertyNamesToHandles.Add(PropertyName, Handle);
 						
 						// Simultaneously edit all selected rows with a similar property
-						FSimpleDelegate OnPropertyValueChanged =
-							FSimpleDelegate::CreateRaw(
+						const TDelegate<void(const FPropertyChangedEvent&)> OnPropertyValueChanged =
+							TDelegate<void(const FPropertyChangedEvent&)>::CreateRaw(
 								this,
 								&SObjectMixerEditorListRow::OnPropertyChanged, PropertyName);
 					
-						Handle->SetOnPropertyValueChanged(OnPropertyValueChanged);
-						Handle->SetOnChildPropertyValueChanged(OnPropertyValueChanged);
+						Handle->SetOnPropertyValueChangedWithData(OnPropertyValueChanged);
+						Handle->SetOnChildPropertyValueChangedWithData(OnPropertyValueChanged);
 
 						return SNew(SBox)
 								.Visibility(EVisibility::SelfHitTestInvisible)
@@ -969,13 +977,14 @@ TSharedPtr<SWidget> SObjectMixerEditorListRow::GenerateCells(
 	return nullptr;
 }
 
-void SObjectMixerEditorListRow::OnPropertyChanged(const FName PropertyName) const
+void SObjectMixerEditorListRow::OnPropertyChanged(const FPropertyChangedEvent& Event, const FName PropertyName) const
 {
 	check(Item.IsValid());
 	
 	auto SetValueOnSelectedItems = [](
 		const FString& ValueAsString, const TArray<FObjectMixerEditorListRowPtr>& OtherSelectedItems,
-		const FName& PropertyName, const FObjectMixerEditorListRowPtr PinnedItem)
+		const FName& PropertyName, const FObjectMixerEditorListRowPtr PinnedItem,
+		const EPropertyValueSetFlags::Type Flags)
 	{
 		if (!ValueAsString.IsEmpty())
 		{
@@ -997,8 +1006,8 @@ void SObjectMixerEditorListRow::OnPropertyChanged(const FName PropertyName) cons
 							{
 								ObjectToModify->Modify();
 							}
-									
-							SelectedHandlePtr->Pin()->SetValueFromFormattedString(ValueAsString);
+							
+							SelectedHandlePtr->Pin()->SetValueFromFormattedString(ValueAsString, Flags);
 						}
 					}
 				}
@@ -1020,7 +1029,17 @@ void SObjectMixerEditorListRow::OnPropertyChanged(const FName PropertyName) cons
 						if (TSharedPtr<IPropertyHandle> PinnedHandle = HandlePtr->Pin())
 						{
 							PinnedHandle->GetValueAsFormattedString(ValueAsString);
-							SetValueOnSelectedItems(ValueAsString, OtherSelectedItems, PropertyName, PinnedItem);
+
+							const bool bIsInteractiveChange = Event.ChangeType == EPropertyChangeType::Interactive;
+							const EPropertyValueSetFlags::Type Flags =
+								 bIsInteractiveChange ? EPropertyValueSetFlags::InteractiveChange : EPropertyValueSetFlags::DefaultFlags;
+							
+							SetValueOnSelectedItems(ValueAsString, OtherSelectedItems, PropertyName, PinnedItem, Flags);
+
+							if (!bIsInteractiveChange)
+							{
+								PinnedItem->GetListViewPtr().Pin()->RequestRebuildList();
+							}
 						}
 					}
 				}
diff --git a/Engine/Plugins/Editor/ObjectMixer/ObjectMixer/Source/ObjectMixer/Private/Views/Widgets/ObjectMixerEditorListMenuContext.cpp b/Engine/Plugins/Editor/ObjectMixer/ObjectMixer/Source/ObjectMixer/Private/Views/Widgets/ObjectMixerEditorListMenuContext.cpp
--- a/Engine/Plugins/Editor/ObjectMixer/ObjectMixer/Source/ObjectMixer/Private/Views/Widgets/ObjectMixerEditorListMenuContext.cpp
+++ b/Engine/Plugins/Editor/ObjectMixer/ObjectMixer/Source/ObjectMixer/Private/Views/Widgets/ObjectMixerEditorListMenuContext.cpp
@@ -25,8 +25,9 @@
 #include "ToolMenuEntry.h"
 #include "ToolMenuSection.h"
 #include "Framework/Commands/GenericCommands.h"
-#include "Widgets/Layout/SBox.h"
+#include "Widgets/Input/SCheckBox.h"
 #include "Widgets/Input/SEditableTextBox.h"
+#include "Widgets/Layout/SBox.h"
 
 #define LOCTEXT_NAMESPACE "ObjectMixerEditorListMenuContext"
 
@@ -287,6 +288,34 @@ void UObjectMixerEditorListMenuContext::RegisterObjectMixerActorContextMenuExten
 	}
 }
 
+void UObjectMixerEditorListMenuContext::AddCollectionWidget(const FName& Key, const FObjectMixerEditorListMenuContextData& ContextData, UToolMenu* Menu)
+{
+	const FText KeyText = FText::FromName(Key);
+
+	TSharedRef<SHorizontalBox> Widget = SNew(SHorizontalBox);
+
+	Widget->AddSlot()
+	.Padding(FMargin(8, 0))
+	.AutoWidth()
+	[
+		SNew(SCheckBox)
+		.OnCheckStateChanged_Lambda([Key, ContextData](ECheckBoxState)
+		{
+			OnClickCollectionMenuEntry(Key, ContextData);
+		})
+		.IsChecked_Static(&UObjectMixerEditorListMenuContext::GetCheckStateForCollection, Key, ContextData)
+	];
+
+	Widget->AddSlot()
+	.AutoWidth()
+	[
+		SNew(STextBlock)
+		.Text(KeyText)
+	];
+								
+	Menu->AddMenuEntry(Key, FToolMenuEntry::InitWidget(Key, Widget, FText(), true));
+}
+
 void UObjectMixerEditorListMenuContext::CreateSelectCollectionsSubMenu(UToolMenu* Menu, FObjectMixerEditorListMenuContextData ContextData)
 {
 	FToolMenuEntry Args;
@@ -313,7 +342,8 @@ void UObjectMixerEditorListMenuContext::CreateSelectCollectionsSubMenu(UToolMenu
 		if (TArray<FName> Collections = MainPanel->GetAllCollectionNames(); Collections.Num() > 0)
 		{
 			FToolMenuSection& Section = Menu->FindOrAddSection("Collections");
-			Section.Label = LOCTEXT("CollectionsSectionName", "Collections");
+			FToolMenuEntry& Separator = Section.AddSeparator("CollectionsSeparator");
+			Separator.Label = LOCTEXT("CollectionsSeparatorLabel", "Collections");
 
 			Collections.StableSort([](const FName& A, const FName B)
 			{
@@ -322,20 +352,7 @@ void UObjectMixerEditorListMenuContext::CreateSelectCollectionsSubMenu(UToolMenu
 							
 			for (const FName& Key : Collections)
 			{
-				const FText KeyText = FText::FromName(Key);
-								
-				Section.AddMenuEntry(
-					Key,
-					KeyText,
-					FText::Format(LOCTEXT("AddObjectsToCollectionTooltipFormat", "Add selected to collection '{0}'"), KeyText),
-					FSlateIcon(),
-					FUIAction(
-						FExecuteAction::CreateStatic(&UObjectMixerEditorListMenuContext::OnClickCollectionMenuEntry, Key, ContextData),
-						FCanExecuteAction(),
-						FIsActionChecked::CreateStatic(&UObjectMixerEditorListMenuContext::AreAllObjectsInCollection, Key, ContextData)
-					),
-					EUserInterfaceActionType::ToggleButton
-				);
+				AddCollectionWidget(Key, ContextData, Menu);
 			}
 		}
 	}
@@ -587,6 +604,45 @@ bool UObjectMixerEditorListMenuContext::AreAllObjectsInCollection(const FName Ke
 	return bAreAllSelectedObjectsInCollection;
 }
 
+ECheckBoxState UObjectMixerEditorListMenuContext::GetCheckStateForCollection(const FName Key,
+	const FObjectMixerEditorListMenuContextData ContextData)
+{
+	const int32 ItemCount = ContextData.SelectedItems.Num();
+	int32 ItemsInCollection = 0;
+	int32 ItemsNotInCollection = 0;
+
+	if (const TSharedPtr<FObjectMixerEditorMainPanel> MainPanel = ContextData.MainPanelPtr.Pin())
+	{
+		for (const TSharedPtr<FObjectMixerEditorListRow>& Item : ContextData.SelectedItems)
+		{
+			if (const UObject* Object = Item->GetObject())
+			{
+				if (MainPanel->IsObjectInCollection(Key, Object))
+				{
+					ItemsInCollection++;
+				}
+				else
+				{
+					ItemsNotInCollection++;
+				}
+			}
+		}
+	}
+
+	ECheckBoxState ReturnValue = ECheckBoxState::Undetermined;
+
+	if (ItemsInCollection == ItemCount)
+	{
+		ReturnValue = ECheckBoxState::Checked;
+	}
+	else if (ItemsNotInCollection == ItemCount)
+	{
+		ReturnValue = ECheckBoxState::Unchecked;
+	}
+
+	return ReturnValue;
+}
+
 FToolMenuEntry UObjectMixerEditorListMenuContext::MakeCustomEditMenu(const FObjectMixerEditorListMenuContextData& ContextData)
 {
 	FToolMenuEntry Entry;
diff --git a/Engine/Plugins/Editor/ObjectMixer/ObjectMixer/Source/ObjectMixer/Public/Views/List/SObjectMixerEditorListRow.h b/Engine/Plugins/Editor/ObjectMixer/ObjectMixer/Source/ObjectMixer/Public/Views/List/SObjectMixerEditorListRow.h
--- a/Engine/Plugins/Editor/ObjectMixer/ObjectMixer/Source/ObjectMixer/Public/Views/List/SObjectMixerEditorListRow.h
+++ b/Engine/Plugins/Editor/ObjectMixer/ObjectMixer/Source/ObjectMixer/Public/Views/List/SObjectMixerEditorListRow.h
@@ -56,7 +56,7 @@ private:
 
 	TSharedPtr<SWidget> GenerateCells(const FName& InColumnName, const TSharedPtr<FObjectMixerEditorListRow> RowPtr);
 
-	void OnPropertyChanged(const FName PropertyName) const;
+	void OnPropertyChanged(const FPropertyChangedEvent& Event, const FName PropertyName) const;
 
 	void OnClickBlueprintLink(UBlueprint* AsBlueprint, UObject* Object = nullptr);
 	FText GetHighlightText() const;
diff --git a/Engine/Plugins/Editor/ObjectMixer/ObjectMixer/Source/ObjectMixer/Public/Views/Widgets/ObjectMixerEditorListMenuContext.h b/Engine/Plugins/Editor/ObjectMixer/ObjectMixer/Source/ObjectMixer/Public/Views/Widgets/ObjectMixerEditorListMenuContext.h
--- a/Engine/Plugins/Editor/ObjectMixer/ObjectMixer/Source/ObjectMixer/Public/Views/Widgets/ObjectMixerEditorListMenuContext.h
+++ b/Engine/Plugins/Editor/ObjectMixer/ObjectMixer/Source/ObjectMixer/Public/Views/Widgets/ObjectMixerEditorListMenuContext.h
@@ -37,6 +37,7 @@ public:
 private:
 
 	static bool DoesSelectionHaveType(const FObjectMixerEditorListMenuContextData& InData, UClass* Type);
+	static void AddCollectionWidget(const FName& Key, const FObjectMixerEditorListMenuContextData& ContextData, UToolMenu* Menu);
 	static void CreateSelectCollectionsSubMenu(UToolMenu* Menu, FObjectMixerEditorListMenuContextData ContextData);
 	static void GenerateMoveToMenu(UToolMenu* InMenu, const FToolMenuInsert& InsertArgs, const FObjectMixerEditorListMenuContextData& ContextData);
 	static void OnFoldersMenuFolderSelected(TSharedRef<ISceneOutlinerTreeItem> Item, FObjectMixerEditorListMenuContextData ContextData);
@@ -51,6 +52,7 @@ private:
 	static void AddObjectsToCollection(const FName Key, const FObjectMixerEditorListMenuContextData ContextData);
 	static void RemoveObjectsFromCollection(const FName Key, const FObjectMixerEditorListMenuContextData ContextData);
 	static bool AreAllObjectsInCollection(const FName Key, const FObjectMixerEditorListMenuContextData ContextData);
+	static ECheckBoxState GetCheckStateForCollection(const FName Key, const FObjectMixerEditorListMenuContextData ContextData);
 
 	static FToolMenuEntry MakeCustomEditMenu(const FObjectMixerEditorListMenuContextData& ContextData);
 	static void ReplaceEditSubMenu(const FObjectMixerEditorListMenuContextData& ContextData);
diff --git a/Engine/Plugins/Editor/ObjectMixer/ObjectMixer/Source/ObjectMixer/Public/Views/Widgets/SHyperlinkWithTextHighlight.h b/Engine/Plugins/Editor/ObjectMixer/ObjectMixer/Source/ObjectMixer/Public/Views/Widgets/SHyperlinkWithTextHighlight.h
--- a/Engine/Plugins/Editor/ObjectMixer/ObjectMixer/Source/ObjectMixer/Public/Views/Widgets/SHyperlinkWithTextHighlight.h
+++ b/Engine/Plugins/Editor/ObjectMixer/ObjectMixer/Source/ObjectMixer/Public/Views/Widgets/SHyperlinkWithTextHighlight.h
@@ -11,6 +11,7 @@
 #include "Styling/SlateTypes.h"
 #include "Widgets/DeclarativeSyntaxSupport.h"
 #include "Widgets/Input/SButton.h"
+#include "Widgets/Text/SInlineEditableTextBlock.h"
 #include "Widgets/Text/STextBlock.h"
 
 #if WITH_ACCESSIBILITY
@@ -29,31 +30,27 @@ public:
 	SLATE_BEGIN_ARGS(SHyperlinkWithTextHighlight)
 		: _Text()
 		, _Style(&FAppStyle::Get().GetWidgetStyle< FHyperlinkStyle >("Hyperlink"))
-		, _TextStyle(nullptr)
 		, _UnderlineStyle(nullptr)
 		, _Padding()
 		, _OnNavigate()
 		, _TextShapingMethod()
 		, _TextFlowDirection()
-		, _HighlightColor()
-		, _HighlightShape()
 		, _HighlightText()
 		{}
 
 		SLATE_ATTRIBUTE( FText, Text )
 		SLATE_STYLE_ARGUMENT( FHyperlinkStyle, Style )
-		SLATE_STYLE_ARGUMENT( FTextBlockStyle, TextStyle )
 		SLATE_STYLE_ARGUMENT( FButtonStyle, UnderlineStyle )
 		SLATE_ATTRIBUTE( FMargin, Padding )
 		SLATE_EVENT( FSimpleDelegate, OnNavigate )
 		SLATE_ARGUMENT( TOptional<ETextShapingMethod>, TextShapingMethod )
 		SLATE_ARGUMENT( TOptional<ETextFlowDirection>, TextFlowDirection )
 
-		/** The color used to highlight the specified text */
-		SLATE_ATTRIBUTE( FLinearColor, HighlightColor )
+		/** Callback to check if the widget is selected, should only be hooked up if parent widget is handling selection or focus. */
+		SLATE_EVENT( FIsSelected, IsSelected )
 
-		/** The brush used to highlight the specified text */
-		SLATE_ATTRIBUTE( const FSlateBrush*, HighlightShape )
+		/** Callback when the text is committed. */
+		SLATE_EVENT( FOnTextCommitted, OnTextCommitted )
 
 		/** Highlight this text in the text block */
 		SLATE_ATTRIBUTE( FText, HighlightText )
@@ -70,8 +67,10 @@ public:
 
 		check (InArgs._Style);
 		const FButtonStyle* UnderlineStyle = InArgs._UnderlineStyle != nullptr ? InArgs._UnderlineStyle : &InArgs._Style->UnderlineStyle;
-		const FTextBlockStyle* TextStyle = InArgs._TextStyle != nullptr ? InArgs._TextStyle : &InArgs._Style->TextStyle;
 		TAttribute<FMargin> Padding = InArgs._Padding.IsSet() ? InArgs._Padding : InArgs._Style->Padding;
+		
+		EditableTextBlockStyle = FAppStyle::Get().GetWidgetStyle<FInlineEditableTextBlockStyle>("InlineEditableTextBlockStyle");
+		EditableTextBlockStyle.TextStyle = InArgs._Style->TextStyle;
 
 		SButton::Construct(
 			SButton::FArguments()
@@ -82,12 +81,12 @@ public:
 			.TextShapingMethod(InArgs._TextShapingMethod)
 			.TextFlowDirection(InArgs._TextFlowDirection)
 			[
-				SNew(STextBlock)
-				.TextStyle(TextStyle)
+				SAssignNew(EditableTextBlock, SInlineEditableTextBlock)
+				.Style(&EditableTextBlockStyle)
 				.Text(InArgs._Text)
-				.HighlightColor(InArgs._HighlightColor)
-				.HighlightShape(InArgs._HighlightShape)
 				.HighlightText(InArgs._HighlightText)
+				.OnTextCommitted(InArgs._OnTextCommitted)
+				.IsSelected(InArgs._IsSelected)
 			]
 		);
 	}
@@ -107,6 +106,8 @@ public:
 	}
 #endif
 
+	TSharedPtr<SInlineEditableTextBlock> EditableTextBlock;
+
 protected:
 
 	/** Invoke the OnNavigate method */
@@ -119,4 +120,6 @@ protected:
 
 	/** The delegate to invoke when someone clicks the hyperlink */
 	FSimpleDelegate OnNavigate;
+
+	FInlineEditableTextBlockStyle EditableTextBlockStyle;
 };
