Object Mixer:

-Light Collection Creation From Non-Light Parents

diff --git a/Engine/Plugins/Editor/ObjectMixer/ObjectMixer/Source/ObjectMixer/Private/Views/Widgets/ObjectMixerEditorListMenuContext.cpp b/Engine/Plugins/Editor/ObjectMixer/ObjectMixer/Source/ObjectMixer/Private/Views/Widgets/ObjectMixerEditorListMenuContext.cpp
--- a/Engine/Plugins/Editor/ObjectMixer/ObjectMixer/Source/ObjectMixer/Private/Views/Widgets/ObjectMixerEditorListMenuContext.cpp
+++ b/Engine/Plugins/Editor/ObjectMixer/ObjectMixer/Source/ObjectMixer/Private/Views/Widgets/ObjectMixerEditorListMenuContext.cpp
@@ -138,6 +138,7 @@ TSharedPtr<SWidget> UObjectMixerEditorListMenuContext::BuildContextMenu(const FO
 	{
 		FLevelEditorContextMenu::RegisterComponentContextMenu();
 		FLevelEditorContextMenu::RegisterElementContextMenu();
+		RegisterObjectMixerElementContextMenuExtension("LevelEditor.ComponentContextMenu");
 
 		PerformLevelEditorRegistrations(Context);
 
@@ -147,6 +148,7 @@ TSharedPtr<SWidget> UObjectMixerEditorListMenuContext::BuildContextMenu(const FO
 	if (DoesSelectionHaveType(InData, UObject::StaticClass()))
 	{
 		FLevelEditorContextMenu::RegisterElementContextMenu();
+		RegisterObjectMixerElementContextMenuExtension("LevelEditor.ElementContextMenu");
 
 		PerformLevelEditorRegistrations(Context);
 
@@ -250,11 +252,24 @@ void UObjectMixerEditorListMenuContext::RegisterFoldersOnlyContextMenu()
 	}));
 }
 
+void UObjectMixerEditorListMenuContext::AddCollectionsMenuItem(UToolMenu* InMenu, const FObjectMixerEditorListMenuContextData& ContextData)
+{
+	FToolMenuSection& Section = InMenu->FindOrAddSection("ObjectMixerCollections");
+	Section.Label = LOCTEXT("ObjectMixerCollectionsSectionName", "Mixer Collections");
+		
+	Section.AddSubMenu(
+		"SelectCollectionsSubMenu",
+		LOCTEXT("SelectCollectionsSubmenu", "Select or Add Collection"),
+		LOCTEXT("SelectCollectionsSubmenu_Tooltip", "Select the collection to which you wish to assign this object."),
+		FNewToolMenuDelegate::CreateStatic(&UObjectMixerEditorListMenuContext::CreateSelectCollectionsSubMenu, ContextData)
+	);
+}
+
 void UObjectMixerEditorListMenuContext::RegisterObjectMixerActorContextMenuExtension()
 {
 	if (UToolMenu* ActorContextMenu = UToolMenus::Get()->ExtendMenu("LevelEditor.ActorContextMenu"))
 	{		
-		ActorContextMenu->AddDynamicSection("DynamicCollectionsSection", FNewToolMenuDelegate::CreateLambda(
+		ActorContextMenu->AddDynamicSection("DynamicActorSection", FNewToolMenuDelegate::CreateLambda(
 			[](UToolMenu* InMenu)
 			{
 				// Ensure proper context
@@ -268,20 +283,32 @@ void UObjectMixerEditorListMenuContext::RegisterObjectMixerActorContextMenuExten
 
 				ReplaceEditSubMenu(ContextData);
 
-				// Collections submenu
+				AddCollectionsMenuItem(InMenu, ContextData);
+
+				GenerateMoveToMenu(InMenu, FToolMenuInsert("ActorTypeTools", EToolMenuInsertType::After), ContextData);
+			}),
+			FToolMenuInsert(NAME_None,EToolMenuInsertType::First)
+		);
+	}
+}
+
+void UObjectMixerEditorListMenuContext::RegisterObjectMixerElementContextMenuExtension(const FName& MenuName)
+{
+	if (UToolMenu* ActorContextMenu = UToolMenus::Get()->ExtendMenu(MenuName))
+	{		
+		ActorContextMenu->AddDynamicSection("DynamicSubobjectSection", FNewToolMenuDelegate::CreateLambda(
+			[](UToolMenu* InMenu)
+			{
+				// Ensure proper context
+				UObjectMixerEditorListMenuContext* Context = InMenu->FindContext<UObjectMixerEditorListMenuContext>();
+				if (!Context || Context->Data.SelectedItems.Num() == 0)
 				{
-					FToolMenuSection& Section = InMenu->FindOrAddSection("ObjectMixerCollections");
-					Section.Label = LOCTEXT("ObjectMixerCollectionsSectionName", "Mixer Collections");
-		
-					Section.AddSubMenu(
-						"SelectCollectionsSubMenu",
-						LOCTEXT("SelectCollectionsSubmenu", "Select or Add Collection"),
-						LOCTEXT("SelectCollectionsSubmenu_Tooltip", "Select the collection to which you wish to assign this object."),
-						FNewToolMenuDelegate::CreateStatic(&UObjectMixerEditorListMenuContext::CreateSelectCollectionsSubMenu, ContextData)
-					);
+					return;
 				}
 
-				GenerateMoveToMenu(InMenu, FToolMenuInsert("ActorTypeTools", EToolMenuInsertType::After), ContextData);
+				const FObjectMixerEditorListMenuContextData& ContextData = Context->Data;
+
+				AddCollectionsMenuItem(InMenu, ContextData);
 			}),
 			FToolMenuInsert(NAME_None,EToolMenuInsertType::First)
 		);
@@ -301,7 +328,7 @@ void UObjectMixerEditorListMenuContext::AddCollectionWidget(const FName& Key, co
 		SNew(SCheckBox)
 		.OnCheckStateChanged_Lambda([Key, ContextData](ECheckBoxState)
 		{
-			OnClickCollectionMenuEntry(Key, ContextData);
+			OnCollectionMenuEntryCheckStateChanged(Key, ContextData);
 		})
 		.IsChecked_Static(&UObjectMixerEditorListMenuContext::GetCheckStateForCollection, Key, ContextData)
 	];
@@ -533,7 +560,7 @@ void UObjectMixerEditorListMenuContext::SelectDescendentsOfSelectedFolders(
 	}
 }
 
-void UObjectMixerEditorListMenuContext::OnClickCollectionMenuEntry(const FName Key, const FObjectMixerEditorListMenuContextData ContextData)
+void UObjectMixerEditorListMenuContext::OnCollectionMenuEntryCheckStateChanged(const FName Key, const FObjectMixerEditorListMenuContextData ContextData)
 {
 	if (AreAllObjectsInCollection(Key, ContextData))
 	{
diff --git a/Engine/Plugins/Editor/ObjectMixer/ObjectMixer/Source/ObjectMixer/Public/Views/List/ObjectMixerEditorListFilters/ObjectMixerEditorListFilter_Collection.h b/Engine/Plugins/Editor/ObjectMixer/ObjectMixer/Source/ObjectMixer/Public/Views/List/ObjectMixerEditorListFilters/ObjectMixerEditorListFilter_Collection.h
--- a/Engine/Plugins/Editor/ObjectMixer/ObjectMixer/Source/ObjectMixer/Public/Views/List/ObjectMixerEditorListFilters/ObjectMixerEditorListFilter_Collection.h
+++ b/Engine/Plugins/Editor/ObjectMixer/ObjectMixer/Source/ObjectMixer/Public/Views/List/ObjectMixerEditorListFilters/ObjectMixerEditorListFilter_Collection.h
@@ -40,7 +40,17 @@ public:
 	{
 		if (InItem.IsValid())
 		{
-			return InItem->IsObjectRefInCollection(CollectionName);
+			if (InItem->IsObjectRefInCollection(CollectionName))
+			{
+				return true;
+			}
+
+			const TWeakPtr<FObjectMixerEditorListRow> ParentRow = InItem->GetDirectParentRow();
+
+			if (ParentRow.IsValid() && ParentRow.Pin()->IsObjectRefInCollection(CollectionName))
+			{
+				return true;
+			}
 		}
 
 		return false;
diff --git a/Engine/Plugins/Editor/ObjectMixer/ObjectMixer/Source/ObjectMixer/Public/Views/Widgets/ObjectMixerEditorListMenuContext.h b/Engine/Plugins/Editor/ObjectMixer/ObjectMixer/Source/ObjectMixer/Public/Views/Widgets/ObjectMixerEditorListMenuContext.h
--- a/Engine/Plugins/Editor/ObjectMixer/ObjectMixer/Source/ObjectMixer/Public/Views/Widgets/ObjectMixerEditorListMenuContext.h
+++ b/Engine/Plugins/Editor/ObjectMixer/ObjectMixer/Source/ObjectMixer/Public/Views/Widgets/ObjectMixerEditorListMenuContext.h
@@ -29,6 +29,7 @@ public:
 	static TSharedPtr<SWidget> BuildContextMenu(const FObjectMixerEditorListMenuContextData& InData);
 	static void RegisterFoldersOnlyContextMenu();
 	static void RegisterObjectMixerActorContextMenuExtension();
+	static void RegisterObjectMixerElementContextMenuExtension(const FName& MenuName);
 
 	FObjectMixerEditorListMenuContextData Data;
 	
@@ -39,6 +40,7 @@ private:
 	static bool DoesSelectionHaveType(const FObjectMixerEditorListMenuContextData& InData, UClass* Type);
 	static void AddCollectionWidget(const FName& Key, const FObjectMixerEditorListMenuContextData& ContextData, UToolMenu* Menu);
 	static void CreateSelectCollectionsSubMenu(UToolMenu* Menu, FObjectMixerEditorListMenuContextData ContextData);
+	static void AddCollectionsMenuItem(UToolMenu* InMenu, const FObjectMixerEditorListMenuContextData& ContextData);
 	static void GenerateMoveToMenu(UToolMenu* InMenu, const FToolMenuInsert& InsertArgs, const FObjectMixerEditorListMenuContextData& ContextData);
 	static void OnFoldersMenuFolderSelected(TSharedRef<ISceneOutlinerTreeItem> Item, FObjectMixerEditorListMenuContextData ContextData);
 	static TSharedRef<TSet<FFolder>> GatherInvalidMoveToDestinations(const FObjectMixerEditorListMenuContextData& ContextData);
@@ -48,7 +50,7 @@ private:
 
 	static void SelectDescendentsOfSelectedFolders(FObjectMixerEditorListMenuContextData ContextData, const bool bRecursive);
 	
-	static void OnClickCollectionMenuEntry(const FName Key, const FObjectMixerEditorListMenuContextData ContextData);
+	static void OnCollectionMenuEntryCheckStateChanged(const FName Key, const FObjectMixerEditorListMenuContextData ContextData);
 	static void AddObjectsToCollection(const FName Key, const FObjectMixerEditorListMenuContextData ContextData);
 	static void RemoveObjectsFromCollection(const FName Key, const FObjectMixerEditorListMenuContextData ContextData);
 	static bool AreAllObjectsInCollection(const FName Key, const FObjectMixerEditorListMenuContextData ContextData);
