Object Mixer:

-Fix: Multi-selection and Issues with Color Picker Panel
-Fix: Multi-Select Blueprint/Actor w/Lights and Light Parameter Change Bug
-Fix: Changing Parameters on a Multi-Selection Only Change Values Of Lights On Screen
-Fix: Blueprint Lights Parameters Switch To "Multiple Values" When Transformed

diff --git a/Engine/Plugins/Editor/ObjectMixer/ObjectMixer/Source/ObjectMixer/Private/Views/List/ObjectMixerEditorListRow.cpp b/Engine/Plugins/Editor/ObjectMixer/ObjectMixer/Source/ObjectMixer/Private/Views/List/ObjectMixerEditorListRow.cpp
--- a/Engine/Plugins/Editor/ObjectMixer/ObjectMixer/Source/ObjectMixer/Private/Views/List/ObjectMixerEditorListRow.cpp
+++ b/Engine/Plugins/Editor/ObjectMixer/ObjectMixer/Source/ObjectMixer/Private/Views/List/ObjectMixerEditorListRow.cpp
@@ -2,15 +2,12 @@
 
 #include "Views/List/ObjectMixerEditorListRow.h"
 
-#include "ObjectMixerEditorSettings.h"
-#include "Views/List/ObjectMixerEditorListFilters/ObjectMixerEditorListFilter_Collection.h"
-
 #include "Algo/AllOf.h"
 #include "ClassIconFinder.h"
 #include "ObjectMixerEditorSerializedData.h"
+#include "ScopedTransaction.h"
 #include "GameFramework/Actor.h"
 #include "Styling/SlateIconFinder.h"
-#include "Views/List/ObjectMixerEditorList.h"
 #include "Views/List/SObjectMixerEditorList.h"
 
 TSharedRef<FObjectMixerListRowDragDropOp> FObjectMixerListRowDragDropOp::New(const TArray<FObjectMixerEditorListRowPtr>& InItems)
@@ -31,6 +28,17 @@ TSharedRef<FObjectMixerListRowDragDropOp> FObjectMixerListRowDragDropOp::New(con
 	return Operation;
 }
 
+FObjectMixerEditorListRow::FTreeItemUniqueIdentifier::FTreeItemUniqueIdentifier(FObjectMixerEditorListRowPtr InTreeViewItem)
+{
+	const UObject* RowObject = InTreeViewItem->GetObject();
+			
+	UniqueIdNumber = RowObject ? RowObject->GetUniqueID() : -1;
+			
+	RowName =
+		InTreeViewItem->GetRowType() == FObjectMixerEditorListRow::Folder ?
+			InTreeViewItem->GetFolderPath().ToString() : InTreeViewItem->GetDisplayName().ToString();
+}
+
 FObjectMixerEditorListRow::~FObjectMixerEditorListRow()
 {
 	FlushReferences();
@@ -532,6 +540,130 @@ void FObjectMixerEditorListRow::ClearSoloRows() const
 	GetListViewPtr().Pin()->ClearSoloRows();
 }
 
+FObjectMixerEditorListRowPtr FObjectMixerEditorListRow::GetHybridChildOrRowItemIfNull()
+{
+	if (FObjectMixerEditorListRowPtr HybridChild = GetHybridChild())
+	{
+		return HybridChild;
+	}
+		
+	return SharedThis(this);
+}
+
+bool FObjectMixerEditorListRow::GetIsItemOrHybridChildSelected()
+{
+	const bool bIsItemSelected = GetIsSelected();
+	const bool bHasHybridChild = GetOrFindHybridRowIndex() != INDEX_NONE && GetChildRows()[GetOrFindHybridRowIndex()].IsValid();
+	const bool bIsChildSelected = bHasHybridChild && GetHybridChild()->GetIsSelected();
+	return bIsChildSelected  || bIsItemSelected;
+}
+
+void SetValueOnSelectedItems(
+	const FString& ValueAsString, const TArray<FObjectMixerEditorListRowPtr>& OtherSelectedItems,
+	const FName& PropertyName, const FObjectMixerEditorListRowPtr PinnedItem,
+	const EPropertyValueSetFlags::Type Flags)
+{
+	if (!ValueAsString.IsEmpty())
+	{
+		FScopedTransaction Transaction(
+			NSLOCTEXT("ObjectMixerEditor","OnPropertyChangedTransaction", "Object Mixer - Bulk Edit Selected Row Properties") );
+				
+		for (const TSharedPtr<FObjectMixerEditorListRow>& SelectedRow : OtherSelectedItems)
+		{
+			const FObjectMixerEditorListRowPtr SelectedHybridRow = SelectedRow->GetHybridChild();
+			const FObjectMixerEditorListRowPtr RowToUse = SelectedHybridRow.IsValid() ? SelectedHybridRow : SelectedRow;
+
+			// Use handles if valid, otherwise use ImportText
+			if (RowToUse == PinnedItem)
+			{
+				return;
+			}
+
+			UObject* ObjectToModify = RowToUse->GetObject();
+			
+			if (const TWeakPtr<IPropertyHandle>* SelectedHandlePtr = RowToUse->PropertyNamesToHandles.Find(PropertyName))
+			{
+				if (SelectedHandlePtr->IsValid())
+				{
+					if (ObjectToModify)
+					{
+						ObjectToModify->Modify();
+					}
+					
+					SelectedHandlePtr->Pin()->SetValueFromFormattedString(ValueAsString, Flags);
+					continue;
+				}
+			}
+			
+			if (ObjectToModify)
+			{
+				if (FProperty* PropertyToChange = FindFProperty<FProperty>(ObjectToModify->GetClass(), PropertyName))
+				{
+					if (void* ValuePtr = PropertyToChange->ContainerPtrToValuePtr<void>(ObjectToModify))
+					{
+						EPropertyChangeType::Type ChangeType =
+							Flags == EPropertyValueSetFlags::InteractiveChange
+								? EPropertyChangeType::Interactive
+								: EPropertyChangeType::ValueSet;
+
+						// Set the actual property value
+						PropertyToChange->ImportText_Direct(*ValueAsString, ValuePtr, ObjectToModify, PPF_None);
+						FPropertyChangedEvent ChangeEvent(
+							PropertyToChange,
+							ChangeType,
+							MakeArrayView({ObjectToModify}));
+						ObjectToModify->PostEditChangeProperty(ChangeEvent);
+
+						// Propagate to outers
+						UObject* Outer = ObjectToModify->GetOuter();
+						while (Outer) 
+						{
+							FPropertyChangedEvent ActorChangeEvent(
+								PropertyToChange,
+								ChangeType,
+								MakeArrayView({Outer}));
+							Outer->PostEditChangeProperty(ActorChangeEvent);
+
+							Outer = Outer->GetOuter();
+						}
+					}
+				}
+			}
+		}
+	}
+}
+
+void FObjectMixerEditorListRow::PropagateChangesToSimilarSelectedRowProperties(const FPropertyPropagationInfo PropertyPropagationInfo)
+{
+	if (PropertyPropagationInfo.PropertyName == NAME_None)
+	{
+		return;
+	}
+
+	if (const FObjectMixerEditorListRowPtr RowToUse = GetHybridChildOrRowItemIfNull())
+	{
+		if (!GetIsItemOrHybridChildSelected())
+		{
+			return;
+		}
+		
+		const TWeakPtr<IPropertyHandle>* HandlePtr = RowToUse->PropertyNamesToHandles.Find(PropertyPropagationInfo.PropertyName);
+		if (HandlePtr->IsValid())
+		{
+			const TArray<FObjectMixerEditorListRowPtr> OtherSelectedItems = RowToUse->GetSelectedTreeViewItems();
+			if (OtherSelectedItems.Num())
+			{
+				FString ValueAsString;
+				(*HandlePtr).Pin()->GetValueAsFormattedString(ValueAsString);
+			
+				SetValueOnSelectedItems(
+					ValueAsString, OtherSelectedItems, PropertyPropagationInfo.PropertyName,
+					RowToUse, PropertyPropagationInfo.PropertyValueSetFlags);
+			}
+		}
+	}
+}
+
 FObjectMixerEditorListRowPtr FObjectMixerEditorListRow::GetAsShared()
 {
 	return SharedThis(this);
diff --git a/Engine/Plugins/Editor/ObjectMixer/ObjectMixer/Source/ObjectMixer/Private/Views/List/SObjectMixerEditorList.cpp b/Engine/Plugins/Editor/ObjectMixer/ObjectMixer/Source/ObjectMixer/Private/Views/List/SObjectMixerEditorList.cpp
--- a/Engine/Plugins/Editor/ObjectMixer/ObjectMixer/Source/ObjectMixer/Private/Views/List/SObjectMixerEditorList.cpp
+++ b/Engine/Plugins/Editor/ObjectMixer/ObjectMixer/Source/ObjectMixer/Private/Views/List/SObjectMixerEditorList.cpp
@@ -206,10 +206,16 @@ void SObjectMixerEditorList::OnRenameCommand()
 	{
 		const FObjectMixerEditorListRowPtr TreeViewItem = GetSelectedTreeViewItems()[0];
 
-		PendingRenameItem = FTreeItemUniqueIdentifier(TreeViewItem);
+		PendingRenameItem = TreeViewItem->GetUniqueIdentifier();
 	}
 }
 
+void SObjectMixerEditorList::AddToPendingPropertyPropagations(
+	const FObjectMixerEditorListRow::FPropertyPropagationInfo& InPropagationInfo)
+{
+	PendingPropertyPropagations.Add(InPropagationInfo);
+}
+
 TArray<FObjectMixerEditorListRowPtr> SObjectMixerEditorList::GetSelectedTreeViewItems() const
 {
 	return TreeViewPtr->GetSelectedItems();
@@ -765,14 +771,14 @@ FListViewColumnInfo* SObjectMixerEditorList::GetColumnInfoByPropertyName(const F
 
 void SObjectMixerEditorList::ExecuteRenameOnPending()
 {
-	using LambdaType = const FObjectMixerEditorListRowPtr*(*)(const TArray<FObjectMixerEditorListRowPtr>&, const FTreeItemUniqueIdentifier&);
+	using LambdaType = const FObjectMixerEditorListRowPtr*(*)(const TArray<FObjectMixerEditorListRowPtr>&, const FObjectMixerEditorListRow::FTreeItemUniqueIdentifier&);
 		
 	static LambdaType RecursivelyFindByUniqueIdentifier = [](
-		const TArray<FObjectMixerEditorListRowPtr>& InObjects, const FTreeItemUniqueIdentifier& InUniqueIdentifier)
+		const TArray<FObjectMixerEditorListRowPtr>& InObjects, const FObjectMixerEditorListRow::FTreeItemUniqueIdentifier& InUniqueIdentifier)
 	{
 		for (const FObjectMixerEditorListRowPtr& TreeViewItem : InObjects)
 		{
-			FTreeItemUniqueIdentifier TreeItemUniqueIdentifier(TreeViewItem);
+			const FObjectMixerEditorListRow::FTreeItemUniqueIdentifier& TreeItemUniqueIdentifier = TreeViewItem->GetUniqueIdentifier();
 
 			if (InUniqueIdentifier == TreeItemUniqueIdentifier)
 			{
@@ -799,7 +805,7 @@ void SObjectMixerEditorList::ExecuteRenameOnPending()
 }
 
 void SObjectMixerEditorList::Tick(const FGeometry& AllottedGeometry, const double InCurrentTime,
-	const float InDeltaTime)
+                                  const float InDeltaTime)
 {
 	SCompoundWidget::Tick(AllottedGeometry, InCurrentTime, InDeltaTime);
 
@@ -822,6 +828,11 @@ void SObjectMixerEditorList::Tick(const FGeometry& AllottedGeometry, const doubl
 		bIsEditorToListSelectionSyncRequested = false;
 	}
 
+	if (PendingPropertyPropagations.Num() > 0)
+	{
+		PropagatePropertyChangesToSelectedRows();
+	}
+
 	if (PendingRenameItem.IdentifiesAnyRow())
 	{
 		ExecuteRenameOnPending();
@@ -1198,7 +1209,7 @@ void SObjectMixerEditorList::CacheTreeState(const TArray<TWeakPtr<IObjectMixerEd
 	{
 		for (const FObjectMixerEditorListRowPtr& TreeViewItem : InObjects)
 		{
-			FTreeItemUniqueIdentifier UniqueIdentifier(TreeViewItem);
+			const FObjectMixerEditorListRow::FTreeItemUniqueIdentifier& UniqueIdentifier = TreeViewItem->GetUniqueIdentifier();
 
 			if (!UniqueIdentifier.RowName.IsEmpty())
 			{
@@ -1251,7 +1262,7 @@ void SObjectMixerEditorList::RestoreTreeState(const TArray<TWeakPtr<IObjectMixer
 	{
 		for (const FObjectMixerEditorListRowPtr& TreeViewItem : InObjects)
 		{
-			const FTreeItemUniqueIdentifier TreeItemUniqueIdentifier(TreeViewItem);
+			const FObjectMixerEditorListRow::FTreeItemUniqueIdentifier& TreeItemUniqueIdentifier = TreeViewItem->GetUniqueIdentifier();
 			
 			if (const FTreeItemStateCache* StateCachePtr = Algo::FindByPredicate(
 				*TreeItemStateCache,
@@ -1273,7 +1284,9 @@ void SObjectMixerEditorList::RestoreTreeState(const TArray<TWeakPtr<IObjectMixer
 			}
 
 			RecursivelyRestoreTreeState(
-				TreeViewItem->GetChildRows(), TreeItemStateCache, InTreeViewPtr, bExpandByDefault);
+				TreeViewItem->GetChildRows(), TreeItemStateCache, InTreeViewPtr,
+				bExpandByDefault
+			);
 		}
 	};
 	
@@ -1288,7 +1301,9 @@ void SObjectMixerEditorList::RestoreTreeState(const TArray<TWeakPtr<IObjectMixer
 		if (FilterComboToStateCaches[CachesItr].FilterCombo == InFilterCombo)
 		{
 			RecursivelyRestoreTreeState(
-			  TreeViewRootObjects, &FilterComboToStateCaches[CachesItr].Caches, TreeViewPtr, bExpandByDefault);
+			  TreeViewRootObjects, &FilterComboToStateCaches[CachesItr].Caches,
+			  TreeViewPtr, bExpandByDefault
+			);
 
 			if (bFlushCache)
 			{
@@ -1300,6 +1315,48 @@ void SObjectMixerEditorList::RestoreTreeState(const TArray<TWeakPtr<IObjectMixer
 	}
 }
 
+void SObjectMixerEditorList::PropagatePropertyChangesToSelectedRows()
+{
+	using LambdaType = void(*)(const TArray<FObjectMixerEditorListRowPtr>&,
+			TSharedPtr<STreeView<FObjectMixerEditorListRowPtr>>,
+			const TArray<FObjectMixerEditorListRow::FPropertyPropagationInfo>&);
+	
+	static LambdaType RecursivelyRestoreTreeState = [](
+		const TArray<FObjectMixerEditorListRowPtr>& InObjects,
+		TSharedPtr<STreeView<FObjectMixerEditorListRowPtr>> InTreeViewPtr,
+		const TArray<FObjectMixerEditorListRow::FPropertyPropagationInfo>& InPendingPropertyPropagations)
+	{
+		for (const FObjectMixerEditorListRowPtr& TreeViewItem : InObjects)
+		{
+			const FObjectMixerEditorListRow::FTreeItemUniqueIdentifier& TreeItemUniqueIdentifier = TreeViewItem->GetUniqueIdentifier();
+
+				if (const FObjectMixerEditorListRow::FPropertyPropagationInfo* PendingPropagation =
+					Algo::FindByPredicate(
+						InPendingPropertyPropagations,
+						[&TreeItemUniqueIdentifier](const FObjectMixerEditorListRow::FPropertyPropagationInfo& Other)
+						{
+							return Other.RowIdentifier == TreeItemUniqueIdentifier;
+						}))
+				{
+					TreeViewItem->PropagateChangesToSimilarSelectedRowProperties(*PendingPropagation);
+				}
+
+			RecursivelyRestoreTreeState(
+				TreeViewItem->GetChildRows(), InTreeViewPtr, InPendingPropertyPropagations
+			);
+		}
+	};
+
+	if (GetSelectedTreeViewItemCount() > 1)
+	{
+		RecursivelyRestoreTreeState(
+		   TreeViewRootObjects, TreeViewPtr, PendingPropertyPropagations
+	   );
+	}
+
+	PendingPropertyPropagations.Empty();
+}
+
 void SObjectMixerEditorList::BuildPerformanceCacheAndGenerateHeaderIfNeeded()
 {
 	TSharedPtr<FObjectMixerEditorList> PinnedListModel = ListModelPtr.Pin();
diff --git a/Engine/Plugins/Editor/ObjectMixer/ObjectMixer/Source/ObjectMixer/Private/Views/List/SObjectMixerEditorListRow.cpp b/Engine/Plugins/Editor/ObjectMixer/ObjectMixer/Source/ObjectMixer/Private/Views/List/SObjectMixerEditorListRow.cpp
--- a/Engine/Plugins/Editor/ObjectMixer/ObjectMixer/Source/ObjectMixer/Private/Views/List/SObjectMixerEditorListRow.cpp
+++ b/Engine/Plugins/Editor/ObjectMixer/ObjectMixer/Source/ObjectMixer/Private/Views/List/SObjectMixerEditorListRow.cpp
@@ -4,7 +4,6 @@
 
 #include "EditorClassUtils.h"
 #include "ObjectMixerEditorStyle.h"
-#include "Views/List/ObjectMixerEditorList.h"
 #include "Views/List/SObjectMixerEditorList.h"
 #include "Views/Widgets/SHyperlinkWithTextHighlight.h"
 
@@ -16,7 +15,6 @@
 #include "Input/DragAndDrop.h"
 #include "ISinglePropertyView.h"
 #include "ObjectMixerEditorLog.h"
-#include "Modules/ModuleManager.h"
 #include "PropertyEditorModule.h"
 #include "PropertyHandle.h"
 #include "ScopedTransaction.h"
@@ -659,14 +657,9 @@ FReply SObjectMixerEditorListRow::HandleAcceptDrop(const FDragDropEvent& DragDro
 
 FObjectMixerEditorListRowPtr SObjectMixerEditorListRow::GetHybridChildOrRowItemIfNull() const
 {
-	if (TSharedPtr<FObjectMixerEditorListRow> PinnedItem = Item.Pin())
-	{
-		if (FObjectMixerEditorListRowPtr HybridChild = HybridRowIndex != INDEX_NONE ? PinnedItem->GetChildRows()[HybridRowIndex] : nullptr)
-		{
-			return HybridChild;
-		}
-			
-		return PinnedItem;
+	if (const TSharedPtr<FObjectMixerEditorListRow> PinnedItem = Item.Pin())
+	{			
+		return PinnedItem->GetHybridChildOrRowItemIfNull();
 	}
 
 	return nullptr;
@@ -674,12 +667,9 @@ FObjectMixerEditorListRowPtr SObjectMixerEditorListRow::GetHybridChildOrRowItemI
 
 bool SObjectMixerEditorListRow::GetIsItemOrHybridChildSelected() const
 {
-	if (Item.IsValid())
-	{
-		const bool bIsItemSelected = Item.Pin()->GetIsSelected();
-		const bool bHasHybridChild = HybridRowIndex != INDEX_NONE && Item.Pin()->GetChildRows()[HybridRowIndex].IsValid();
-		const bool bIsChildSelected = bHasHybridChild && Item.Pin()->GetChildRows()[HybridRowIndex]->GetIsSelected();
-		return bIsChildSelected  || bIsItemSelected;
+	if (const TSharedPtr<FObjectMixerEditorListRow> PinnedItem = Item.Pin())
+	{			
+		return PinnedItem->GetIsItemOrHybridChildSelected();
 	}
 
 	return false;
@@ -870,6 +860,7 @@ TSharedPtr<SWidget> SObjectMixerEditorListRow::GenerateCells(
 	
 	if (InColumnName.IsEqual(SObjectMixerEditorList::ItemNameColumnName))
 	{
+		// Pass in actual row item even if hybrid row
 		return SNew(SInlineEditableRowNameCellWidget,
 			Item.Pin().ToSharedRef(), HybridRowIndex != INDEX_NONE ? GetHybridChildOrRowItemIfNull() : nullptr);
 	}
@@ -979,70 +970,26 @@ TSharedPtr<SWidget> SObjectMixerEditorListRow::GenerateCells(
 
 void SObjectMixerEditorListRow::OnPropertyChanged(const FPropertyChangedEvent& Event, const FName PropertyName) const
 {
-	check(Item.IsValid());
-	
-	auto SetValueOnSelectedItems = [](
-		const FString& ValueAsString, const TArray<FObjectMixerEditorListRowPtr>& OtherSelectedItems,
-		const FName& PropertyName, const FObjectMixerEditorListRowPtr PinnedItem,
-		const EPropertyValueSetFlags::Type Flags)
+	if (const TSharedPtr<FObjectMixerEditorListRow> PinnedItem = Item.Pin())
 	{
-		if (!ValueAsString.IsEmpty())
-		{
-			FScopedTransaction Transaction(
-				LOCTEXT("OnPropertyChangedTransaction", "Object Mixer - Bulk Edit Selected Row Properties") );
-					
-			for (const TSharedPtr<FObjectMixerEditorListRow>& SelectedRow : OtherSelectedItems)
-			{
-				const FObjectMixerEditorListRowPtr SelectedHybridRow = SelectedRow->GetHybridChild();
-				const FObjectMixerEditorListRowPtr RowToUse = SelectedHybridRow.IsValid() ? SelectedHybridRow : SelectedRow;
+		const EPropertyValueSetFlags::Type Flag =
+			 Event.ChangeType == EPropertyChangeType::Interactive ?
+			 	EPropertyValueSetFlags::InteractiveChange : EPropertyValueSetFlags::DefaultFlags;
 
-				if (RowToUse != PinnedItem)
-				{
-					if (const TWeakPtr<IPropertyHandle>* SelectedHandlePtr = RowToUse->PropertyNamesToHandles.Find(PropertyName))
-					{
-						if (SelectedHandlePtr->IsValid())
-						{
-							if (UObject* ObjectToModify = RowToUse->GetObject())
-							{
-								ObjectToModify->Modify();
-							}
-							
-							SelectedHandlePtr->Pin()->SetValueFromFormattedString(ValueAsString, Flags);
-						}
-					}
-				}
-			}
+		const FObjectMixerEditorListRow::FPropertyPropagationInfo PropagationInfo(
+			{PinnedItem->GetUniqueIdentifier(), PropertyName, Flag});
+		
+		if (Flag == EPropertyValueSetFlags::InteractiveChange)
+		{
+			PinnedItem->PropagateChangesToSimilarSelectedRowProperties(PropagationInfo);
 		}
-	};
-	
-	if (PropertyName != NAME_None)
-	{
-		if (const FObjectMixerEditorListRowPtr PinnedItem = GetHybridChildOrRowItemIfNull())
+		else
 		{
-			if (GetIsItemOrHybridChildSelected())
+			// If not an interactive change, schedule property propagation on next frame
+			if (const TSharedPtr<SObjectMixerEditorList> PinnedListView = PinnedItem->GetListViewPtr().Pin())
 			{
-				if (const TWeakPtr<IPropertyHandle>* HandlePtr = PinnedItem->PropertyNamesToHandles.Find(PropertyName); HandlePtr->IsValid())
-				{
-					if (TArray<FObjectMixerEditorListRowPtr> OtherSelectedItems = PinnedItem->GetSelectedTreeViewItems(); OtherSelectedItems.Num())
-					{
-						FString ValueAsString;
-						if (TSharedPtr<IPropertyHandle> PinnedHandle = HandlePtr->Pin())
-						{
-							PinnedHandle->GetValueAsFormattedString(ValueAsString);
-
-							const bool bIsInteractiveChange = Event.ChangeType == EPropertyChangeType::Interactive;
-							const EPropertyValueSetFlags::Type Flags =
-								 bIsInteractiveChange ? EPropertyValueSetFlags::InteractiveChange : EPropertyValueSetFlags::DefaultFlags;
-							
-							SetValueOnSelectedItems(ValueAsString, OtherSelectedItems, PropertyName, PinnedItem, Flags);
-
-							if (!bIsInteractiveChange)
-							{
-								PinnedItem->GetListViewPtr().Pin()->RequestRebuildList();
-							}
-						}
-					}
-				}
+				PinnedListView->AddToPendingPropertyPropagations(PropagationInfo);
+				PinnedListView->RequestRebuildList();
 			}
 		}
 	}
diff --git a/Engine/Plugins/Editor/ObjectMixer/ObjectMixer/Source/ObjectMixer/Public/Views/List/ObjectMixerEditorListRow.h b/Engine/Plugins/Editor/ObjectMixer/ObjectMixer/Source/ObjectMixer/Public/Views/List/ObjectMixerEditorListRow.h
--- a/Engine/Plugins/Editor/ObjectMixer/ObjectMixer/Source/ObjectMixer/Public/Views/List/ObjectMixerEditorListRow.h
+++ b/Engine/Plugins/Editor/ObjectMixer/ObjectMixer/Source/ObjectMixer/Public/Views/List/ObjectMixerEditorListRow.h
@@ -64,6 +64,46 @@ struct OBJECTMIXEREDITOR_API FObjectMixerEditorListRow : TSharedFromThis<FObject
 		bool bShouldBeSolo = false;
 	};
 
+	struct FTreeItemUniqueIdentifier
+	{
+		FTreeItemUniqueIdentifier() = default;
+		
+		FTreeItemUniqueIdentifier(FObjectMixerEditorListRowPtr InTreeViewItem);
+
+		bool operator==(const FTreeItemUniqueIdentifier& Other) const
+		{
+			if (Other.UniqueIdNumber != 0)
+			{
+				return Other.UniqueIdNumber == UniqueIdNumber;
+			}
+
+			return Other.RowName.Equals(RowName);
+		}
+
+		bool IdentifiesAnyRow() const
+		{
+			return UniqueIdNumber != 0 || !RowName.IsEmpty();
+		}
+
+		void Reset()
+		{
+			UniqueIdNumber = 0;
+			RowName = "";
+		}
+		
+		uint32 UniqueIdNumber = 0;
+		FString RowName = "";
+	};
+
+	struct FPropertyPropagationInfo
+	{
+		FPropertyPropagationInfo() = default;
+		
+		FTreeItemUniqueIdentifier RowIdentifier;
+		FName PropertyName = NAME_None;
+		EPropertyValueSetFlags::Type PropertyValueSetFlags = 0;
+	};
+
 	~FObjectMixerEditorListRow();
 
 	bool operator==(const TSharedPtr<FObjectMixerEditorListRow>& Other) const
@@ -236,6 +276,12 @@ struct OBJECTMIXEREDITOR_API FObjectMixerEditorListRow : TSharedFromThis<FObject
 	void SetRowSoloState(const bool bNewSolo);
 
 	void ClearSoloRows() const;
+	
+	FObjectMixerEditorListRowPtr GetHybridChildOrRowItemIfNull();
+
+	bool GetIsItemOrHybridChildSelected();
+
+	void PropagateChangesToSimilarSelectedRowProperties(const FPropertyPropagationInfo PropertyPropagationInfo);
 
 	FOnRenameCommand& OnRenameCommand()
 	{
@@ -246,6 +292,15 @@ struct OBJECTMIXEREDITOR_API FObjectMixerEditorListRow : TSharedFromThis<FObject
 	{
 		OnRenameCommandDelegate.ExecuteIfBound();
 	}
+
+	const FTreeItemUniqueIdentifier& GetUniqueIdentifier()
+	{
+		if (!UniqueIdentifier.IdentifiesAnyRow())
+		{
+			UniqueIdentifier = FTreeItemUniqueIdentifier(SharedThis(this));
+		}
+		return UniqueIdentifier;
+	}
 	
 	TMap<FName, TWeakPtr<IPropertyHandle>> PropertyNamesToHandles;
 
@@ -280,4 +335,6 @@ private:
 	int32 CachedHybridRowIndex = INDEX_NONE;
 
 	FOnRenameCommand OnRenameCommandDelegate;
+
+	FTreeItemUniqueIdentifier UniqueIdentifier;
 };
diff --git a/Engine/Plugins/Editor/ObjectMixer/ObjectMixer/Source/ObjectMixer/Public/Views/List/SObjectMixerEditorList.h b/Engine/Plugins/Editor/ObjectMixer/ObjectMixer/Source/ObjectMixer/Public/Views/List/SObjectMixerEditorList.h
--- a/Engine/Plugins/Editor/ObjectMixer/ObjectMixer/Source/ObjectMixer/Public/Views/List/SObjectMixerEditorList.h
+++ b/Engine/Plugins/Editor/ObjectMixer/ObjectMixer/Source/ObjectMixer/Public/Views/List/SObjectMixerEditorList.h
@@ -103,6 +103,8 @@ public:
 	/** Called when the Rename command is executed from the UI or hotkey. */
 	void OnRenameCommand();
 
+	void AddToPendingPropertyPropagations(const FObjectMixerEditorListRow::FPropertyPropagationInfo& InPropagationInfo);
+
 	[[nodiscard]] TArray<FObjectMixerEditorListRowPtr> GetSelectedTreeViewItems() const;
 	int32 GetSelectedTreeViewItemCount() const;
 
@@ -163,6 +165,8 @@ public:
 	void CacheTreeState(const TArray<TWeakPtr<IObjectMixerEditorListFilter>>& InFilterCombo);
 	void RestoreTreeState(const TArray<TWeakPtr<IObjectMixerEditorListFilter>>& InFilterCombo, const bool bFlushCache = true);
 
+	void PropagatePropertyChangesToSelectedRows();
+
 	// Sorting
 
 	const FName& GetActiveSortingColumnName() const
@@ -256,49 +260,9 @@ protected:
 
 	TSharedPtr<STreeView<FObjectMixerEditorListRowPtr>> TreeViewPtr;
 
-	struct FTreeItemUniqueIdentifier
-	{
-		FTreeItemUniqueIdentifier() = default;
-		
-		FTreeItemUniqueIdentifier(FObjectMixerEditorListRowPtr InTreeViewItem)
-		{
-			const UObject* RowObject = InTreeViewItem->GetObject();
-			
-			UniqueIdNumber = RowObject ? RowObject->GetUniqueID() : -1;
-			
-			RowName =
-				InTreeViewItem->GetRowType() == FObjectMixerEditorListRow::Folder ?
-					InTreeViewItem->GetFolderPath().ToString() : InTreeViewItem->GetDisplayName().ToString();
-		}
-
-		bool operator==(const FTreeItemUniqueIdentifier& Other) const
-		{
-			if (Other.UniqueIdNumber != -1)
-			{
-				return Other.UniqueIdNumber == UniqueIdNumber;
-			}
-
-			return Other.RowName.Equals(RowName);
-		}
-
-		bool IdentifiesAnyRow() const
-		{
-			return UniqueIdNumber != 1 || !RowName.IsEmpty();
-		}
-
-		void Reset()
-		{
-			UniqueIdNumber = -1;
-			RowName = "";
-		}
-		
-		uint32 UniqueIdNumber = -1;
-		FString RowName = "";
-	};
-
 	struct FTreeItemStateCache
 	{
-		FTreeItemUniqueIdentifier UniqueIdentifier;
+		FObjectMixerEditorListRow::FTreeItemUniqueIdentifier UniqueIdentifier;
 		bool bIsExpanded = false;
 		bool bIsSelected = false;
 		FObjectMixerEditorListRow::FTransientEditorVisibilityRules VisibilityRules;
@@ -312,7 +276,8 @@ protected:
 
 	TArray<FFilterComboToStateCaches> FilterComboToStateCaches;
 	
-	FTreeItemUniqueIdentifier PendingRenameItem;
+	FObjectMixerEditorListRow::FTreeItemUniqueIdentifier PendingRenameItem;
+	TArray<FObjectMixerEditorListRow::FPropertyPropagationInfo> PendingPropertyPropagations;
 	
 	TMap<UObject*, FObjectMixerEditorListRowPtr> ObjectsToRowsCreated;
 
diff --git a/Engine/Plugins/Editor/ObjectMixer/ObjectMixer/Source/ObjectMixer/Public/Views/List/SObjectMixerEditorListRow.h b/Engine/Plugins/Editor/ObjectMixer/ObjectMixer/Source/ObjectMixer/Public/Views/List/SObjectMixerEditorListRow.h
--- a/Engine/Plugins/Editor/ObjectMixer/ObjectMixer/Source/ObjectMixer/Public/Views/List/SObjectMixerEditorListRow.h
+++ b/Engine/Plugins/Editor/ObjectMixer/ObjectMixer/Source/ObjectMixer/Public/Views/List/SObjectMixerEditorListRow.h
@@ -6,7 +6,6 @@
 
 #include "DragAndDrop/DecoratedDragDropOp.h"
 #include "Widgets/Input/SNumericEntryBox.h"
-#include "Widgets/IToolTip.h"
 #include "Widgets/Views/STableRow.h"
 
 class SObjectMixerEditorListValueInput;
@@ -57,9 +56,6 @@ private:
 	TSharedPtr<SWidget> GenerateCells(const FName& InColumnName, const TSharedPtr<FObjectMixerEditorListRow> RowPtr);
 
 	void OnPropertyChanged(const FPropertyChangedEvent& Event, const FName PropertyName) const;
-
-	void OnClickBlueprintLink(UBlueprint* AsBlueprint, UObject* Object = nullptr);
-	FText GetHighlightText() const;
 	
 	TWeakPtr<FObjectMixerEditorListRow> Item;
 	
