Object Mixer:

-Implement Multi-filter

diff --git a/Engine/Plugins/Editor/ObjectMixer/ObjectMixer/Source/ObjectMixer/Private/ObjectMixerEditorModule.cpp b/Engine/Plugins/Editor/ObjectMixer/ObjectMixer/Source/ObjectMixer/Private/ObjectMixerEditorModule.cpp
--- a/Engine/Plugins/Editor/ObjectMixer/ObjectMixer/Source/ObjectMixer/Private/ObjectMixerEditorModule.cpp
+++ b/Engine/Plugins/Editor/ObjectMixer/ObjectMixer/Source/ObjectMixer/Private/ObjectMixerEditorModule.cpp
@@ -292,7 +292,7 @@ TSharedRef<SDockTab> FObjectMixerEditorModule::SpawnMainPanelTab()
 
 	if (DefaultFilterClass)
 	{
-		MainPanel->OnClassSelectionChanged(DefaultFilterClass);
+		MainPanel->SetDefaultFilterClass(DefaultFilterClass);
 	}
 			
 	return DockTab;
@@ -386,4 +386,9 @@ void FObjectMixerEditorModule::BindDelegates()
 	});
 }
 
+const TSubclassOf<UObjectMixerObjectFilter>& FObjectMixerEditorModule::GetDefaultFilterClass() const
+{
+	return DefaultFilterClass;
+}
+
 #undef LOCTEXT_NAMESPACE
diff --git a/Engine/Plugins/Editor/ObjectMixer/ObjectMixer/Source/ObjectMixer/Private/Views/List/ObjectMixerEditorListRow.cpp b/Engine/Plugins/Editor/ObjectMixer/ObjectMixer/Source/ObjectMixer/Private/Views/List/ObjectMixerEditorListRow.cpp
--- a/Engine/Plugins/Editor/ObjectMixer/ObjectMixer/Source/ObjectMixer/Private/Views/List/ObjectMixerEditorListRow.cpp
+++ b/Engine/Plugins/Editor/ObjectMixer/ObjectMixer/Source/ObjectMixer/Private/Views/List/ObjectMixerEditorListRow.cpp
@@ -52,7 +52,18 @@ void FObjectMixerEditorListRow::FlushReferences()
 	}
 }
 
-UObjectMixerObjectFilter* FObjectMixerEditorListRow::GetObjectFilter() const
+const TArray<TObjectPtr<UObjectMixerObjectFilter>>& FObjectMixerEditorListRow::GetObjectFilterInstances() const
+{
+	const TSharedPtr<SObjectMixerEditorList> PinnedListView = GetListViewPtr().Pin();
+	const TSharedPtr<FObjectMixerEditorList> PinnedListModel = PinnedListView->GetListModelPtr().Pin();
+	const TSharedPtr<FObjectMixerEditorMainPanel> PinnedMainPanel = PinnedListModel->GetMainPanelModel().Pin();
+	check (PinnedMainPanel);
+
+
+	return PinnedMainPanel->GetObjectFilterInstances();
+}
+
+const UObjectMixerObjectFilter* FObjectMixerEditorListRow::GetMainObjectFilterInstance() const
 {
 	if (const TSharedPtr<SObjectMixerEditorList> PinnedListView = GetListViewPtr().Pin())
 	{
@@ -60,14 +71,12 @@ UObjectMixerObjectFilter* FObjectMixerEditorListRow::GetObjectFilter() const
 		{
 			if (const TSharedPtr<FObjectMixerEditorMainPanel> PinnedMainPanel = PinnedListModel->GetMainPanelModel().Pin())
 			{
-				if (UObjectMixerObjectFilter* Filter = PinnedMainPanel->GetObjectFilter())
-				{
-					return Filter;
-				}
+				return PinnedMainPanel->GetMainObjectFilterInstance();
 			}
 		}
 	}
 
+
 	return nullptr;
 }
 
@@ -79,14 +88,11 @@ bool FObjectMixerEditorListRow::IsObjectRefInCollection(const FName& CollectionN
 		{
 			return true;
 		}
-		
-		UObjectMixerEditorSerializedData* SerializedData = GetMutableDefault<UObjectMixerEditorSerializedData>();
-		check(SerializedData);
-		
-		if (const UObjectMixerObjectFilter* Filter = GetObjectFilter())
+
+		if (const TSharedPtr<FObjectMixerEditorMainPanel> MainPanel =
+			GetListViewPtr().Pin()->GetListModelPtr().Pin()->GetMainPanelModel().Pin())
 		{
-			return SerializedData->IsObjectInCollection(
-				Filter->GetClass()->GetFName(), CollectionName, GetObject());
+			return MainPanel->IsObjectInCollection(CollectionName, GetObject());
 		}
 	}
 	
@@ -261,9 +267,9 @@ bool FObjectMixerEditorListRow::MatchSearchTokensToSearchTerms(
 
 		if (const TObjectPtr<UObject> Object = GetObject())
 		{			
-			if (const UObjectMixerObjectFilter* Filter = GetObjectFilter())
+			for (const TObjectPtr<UObjectMixerObjectFilter>& Filter : GetObjectFilterInstances())
 			{
-				CachedSearchTerms += Filter->GetRowDisplayName(Object).ToString();
+				CachedSearchTerms = CachedSearchTerms + " " + Filter->GetRowDisplayName(Object).ToString();
 			}
 		}
 	}
@@ -406,7 +412,7 @@ FText FObjectMixerEditorListRow::GetDisplayName(const bool bIsHybridRow) const
 		return Override;
 	}
 
-	if (const UObjectMixerObjectFilter* Filter = GetObjectFilter())
+	if (const UObjectMixerObjectFilter* Filter = GetMainObjectFilterInstance())
 	{
 		if (const TObjectPtr<UObject> Object = GetObject())
 		{
@@ -491,7 +497,7 @@ bool FObjectMixerEditorListRow::GetCurrentEditorObjectVisibility()
 		return false;
 	}
 	
-	if (const UObjectMixerObjectFilter* Filter = GetObjectFilter())
+	if (const UObjectMixerObjectFilter* Filter = GetMainObjectFilterInstance())
 	{
 		return Filter->GetRowEditorVisibility(GetObject());
 	}
@@ -501,7 +507,7 @@ bool FObjectMixerEditorListRow::GetCurrentEditorObjectVisibility()
 
 void FObjectMixerEditorListRow::SetCurrentEditorObjectVisibility(const bool bNewIsVisible, const bool bIsRecursive)
 {
-	if (const UObjectMixerObjectFilter* Filter = GetObjectFilter())
+	if (const UObjectMixerObjectFilter* Filter = GetMainObjectFilterInstance())
 	{
 		Filter->OnSetRowEditorVisibility(GetObject(), bNewIsVisible);
 
diff --git a/Engine/Plugins/Editor/ObjectMixer/ObjectMixer/Source/ObjectMixer/Private/Views/List/SObjectMixerEditorList.cpp b/Engine/Plugins/Editor/ObjectMixer/ObjectMixer/Source/ObjectMixer/Private/Views/List/SObjectMixerEditorList.cpp
--- a/Engine/Plugins/Editor/ObjectMixer/ObjectMixer/Source/ObjectMixer/Private/Views/List/SObjectMixerEditorList.cpp
+++ b/Engine/Plugins/Editor/ObjectMixer/ObjectMixer/Source/ObjectMixer/Private/Views/List/SObjectMixerEditorList.cpp
@@ -885,9 +885,9 @@ TSharedRef<SWidget> SObjectMixerEditorList::GenerateHeaderRowContextMenu() const
 
 					if (UObjectMixerEditorSerializedData* SerializedData = GetMutableDefault<UObjectMixerEditorSerializedData>())
 					{
-						if (const TSubclassOf<UObjectMixerObjectFilter> Filter = PinnedListModel->GetMainPanelModel().Pin()->GetObjectFilterClass())
+						for (const TSubclassOf<UObjectMixerObjectFilter>& FilterClass : PinnedListModel->GetMainPanelModel().Pin()->GetObjectFilterClasses())
 						{
-							const FName FilterName = Filter->GetFName();
+							const FName FilterName = FilterClass->GetFName();
 							SerializedData->SetShouldShowColumn(FilterName, PropertyName, bNewColumnEnabled);
 						}
 					}
@@ -1019,9 +1019,15 @@ TSharedPtr<SHeaderRow> SObjectMixerEditorList::GenerateHeaderRow()
 	HeaderRow->ClearColumns();
 	ListViewColumns.Empty(ListViewColumns.Num());
 	
-	TSet<UClass*> SpecifiedClasses =
-		PinnedListModel->GetMainPanelModel().Pin()->GetObjectFilter()->GetParentAndChildClassesFromSpecifiedClasses(
-			ObjectClassesToFilterCache, PropertyInheritanceInclusionOptionsCache);
+	TSet<UClass*> SpecifiedClasses;
+	for (const TObjectPtr<UObjectMixerObjectFilter>& Instance : PinnedListModel->GetMainPanelModel().Pin()->GetObjectFilterInstances())
+	{
+		SpecifiedClasses.Append(
+			Instance->GetParentAndChildClassesFromSpecifiedClasses(
+				ObjectClassesToFilterCache, PropertyInheritanceInclusionOptionsCache
+			)
+		);
+	}
 	
 	for (const UClass* Class : SpecifiedClasses)
 	{
@@ -1146,12 +1152,17 @@ TSharedPtr<SHeaderRow> SObjectMixerEditorList::GenerateHeaderRow()
 			// Then load visible columns from SerializedData
 			if (UObjectMixerEditorSerializedData* SerializedData = GetMutableDefault<UObjectMixerEditorSerializedData>())
 			{
-				if (const TSubclassOf<UObjectMixerObjectFilter> Filter = PinnedListModel->GetMainPanelModel().Pin()->GetObjectFilterClass())
+				for (const TSubclassOf<UObjectMixerObjectFilter>& Filter : PinnedListModel->GetMainPanelModel().Pin()->GetObjectFilterClasses())
 				{
 					const FName FilterName = Filter->GetFName();
 					if (SerializedData->IsColumnDataSerialized(FilterName, ColumnInfo.PropertyName))
 					{
-						bShouldShowColumn = SerializedData->ShouldShowColumn(FilterName, ColumnInfo.PropertyName);
+						if (SerializedData->ShouldShowColumn(FilterName, ColumnInfo.PropertyName))
+						{
+							// If any filter's serialized data says we should show a column, show it and stomp other filters' saved setting.
+							bShouldShowColumn = true;
+							break;
+						}
 					}
 				}
 			}
@@ -1360,77 +1371,82 @@ void SObjectMixerEditorList::PropagatePropertyChangesToSelectedRows()
 void SObjectMixerEditorList::BuildPerformanceCacheAndGenerateHeaderIfNeeded()
 {
 	TSharedPtr<FObjectMixerEditorList> PinnedListModel = ListModelPtr.Pin();
-	check(PinnedListModel && PinnedListModel->GetMainPanelModel().IsValid());
-	
-	// If any of the following overrides change, we need to regenerate the header row. Otherwise skip regeneration for performance reasons.
+	check(PinnedListModel);
+	TSharedPtr<FObjectMixerEditorMainPanel> PinnedMainPanelModel = PinnedListModel->GetMainPanelModel().Pin();
+	check(PinnedMainPanelModel);
+
+	// If any of the following overrides change, we need to regenerate the header row. Otherwise skip regeneration for performance reasons.	
 	// GetObjectClassesToFilter, GetColumnsToShowByDefault, GetColumnsToExclude,
 	// GetForceAddedColumns, GetObjectMixerPropertyInheritanceInclusionOptions, ShouldIncludeUnsupportedProperties
 	bool bNeedToGenerateHeaders = false;
 	
-	const TObjectPtr<UObjectMixerObjectFilter> SelectedFilter = PinnedListModel->GetMainPanelModel().Pin()->GetObjectFilter();
-	if (!SelectedFilter)
+	TSet<UClass*> LocalObjectClassesToFilterCache;
+	TSet<FName> LocalColumnsToShowByDefaultCache;
+	TSet<FName> LocalColumnsToExcludeCache;
+	TSet<FName> LocalForceAddedColumnsCache;
+	
+	for (const TObjectPtr<UObjectMixerObjectFilter>& FilterInstance : PinnedMainPanelModel->GetObjectFilterInstances())
 	{
-		UE_LOG(LogObjectMixerEditor, Display, TEXT("%hs: No classes defined in UObjectMixerObjectFilter class."), __FUNCTION__);
-		return;
+		if (!FilterInstance)
+		{
+			UE_LOG(LogObjectMixerEditor, Display, TEXT("%hs: UObjectMixerObjectFilter instance not valid."), __FUNCTION__);
+			continue;
+		}
+
+		LocalObjectClassesToFilterCache.Append(FilterInstance->GetObjectClassesToFilter());
+		LocalColumnsToShowByDefaultCache.Append(FilterInstance->GetColumnsToShowByDefault());
+		LocalColumnsToExcludeCache.Append(FilterInstance->GetColumnsToExclude());
+		LocalForceAddedColumnsCache.Append(FilterInstance->GetForceAddedColumns());
 	}
 
-	if (const TSet<UClass*> ObjectClassesToFilter = SelectedFilter->GetObjectClassesToFilter();
-		ObjectClassesToFilter.Num() != ObjectClassesToFilterCache.Num() ||
-		ObjectClassesToFilter.Difference(ObjectClassesToFilterCache).Num() > 0 || ObjectClassesToFilterCache.Difference(ObjectClassesToFilter).Num() > 0)
+	if (LocalObjectClassesToFilterCache.Difference(ObjectClassesToFilterCache).Num() > 0 || 
+		ObjectClassesToFilterCache.Difference(LocalObjectClassesToFilterCache).Num() > 0)
 	{
-		ObjectClassesToFilterCache = ObjectClassesToFilter;
-		if (!bNeedToGenerateHeaders)
-		{
-			bNeedToGenerateHeaders = true;
-		}
+		ObjectClassesToFilterCache = LocalObjectClassesToFilterCache;
+		bNeedToGenerateHeaders = true;
 	}
-	if (const TSet<FName> ColumnsToShowByDefault = SelectedFilter->GetColumnsToShowByDefault();
-		ColumnsToShowByDefault.Num() != ColumnsToShowByDefaultCache.Num() ||
-		ColumnsToShowByDefault.Difference(ColumnsToShowByDefaultCache).Num() > 0 || ColumnsToShowByDefaultCache.Difference(ColumnsToShowByDefault).Num() > 0)
+
+	if (LocalColumnsToShowByDefaultCache.Difference(ColumnsToShowByDefaultCache).Num() > 0 || 
+		ColumnsToShowByDefaultCache.Difference(LocalColumnsToShowByDefaultCache).Num() > 0)
 	{
-		ColumnsToShowByDefaultCache = ColumnsToShowByDefault;
-		if (!bNeedToGenerateHeaders)
-		{
-			bNeedToGenerateHeaders = true;
-		}
+		ColumnsToShowByDefaultCache = LocalColumnsToShowByDefaultCache;
+		bNeedToGenerateHeaders = true;
 	}
-	if (const TSet<FName> ColumnsToExclude = SelectedFilter->GetColumnsToExclude();
-		ColumnsToExclude.Num() != ColumnsToExcludeCache.Num() ||
-		ColumnsToExclude.Difference(ColumnsToExcludeCache).Num() > 0 || ColumnsToExcludeCache.Difference(ColumnsToExclude).Num() > 0)
+
+	if (LocalColumnsToExcludeCache.Difference(ColumnsToExcludeCache).Num() > 0 || 
+		ColumnsToExcludeCache.Difference(LocalColumnsToExcludeCache).Num() > 0)
 	{
-		ColumnsToExcludeCache = ColumnsToExclude;
-		if (!bNeedToGenerateHeaders)
-		{
-			bNeedToGenerateHeaders = true;
-		}
+		ColumnsToExcludeCache = LocalColumnsToExcludeCache;
+		bNeedToGenerateHeaders = true;
+	}
+
+	if (LocalForceAddedColumnsCache.Difference(ForceAddedColumnsCache).Num() > 0 || 
+		ForceAddedColumnsCache.Difference(LocalForceAddedColumnsCache).Num() > 0)
+	{
+		ForceAddedColumnsCache = LocalForceAddedColumnsCache;
+		bNeedToGenerateHeaders = true;
 	}
-	if (const TSet<FName> ForceAddedColumns = SelectedFilter->GetForceAddedColumns();
-		ForceAddedColumns.Num() != ForceAddedColumnsCache.Num() ||
-		ForceAddedColumns.Difference(ForceAddedColumnsCache).Num() > 0 || ForceAddedColumnsCache.Difference(ForceAddedColumns).Num() > 0)
+
+	// These properties should be governed by the Main instance
+	if (const UObjectMixerObjectFilter* FilterInstance = PinnedMainPanelModel->GetMainObjectFilterInstance())
 	{
-		ForceAddedColumnsCache = ForceAddedColumns;
-		if (!bNeedToGenerateHeaders)
+		if (const EObjectMixerInheritanceInclusionOptions PropertyInheritanceInclusionOptions =
+			FilterInstance->GetObjectMixerPropertyInheritanceInclusionOptions(); 
+			PropertyInheritanceInclusionOptions != PropertyInheritanceInclusionOptionsCache)
 		{
+			PropertyInheritanceInclusionOptionsCache = PropertyInheritanceInclusionOptions;
 			bNeedToGenerateHeaders = true;
 		}
-	}
-	if (const EObjectMixerInheritanceInclusionOptions PropertyInheritanceInclusionOptions = SelectedFilter->GetObjectMixerPropertyInheritanceInclusionOptions(); 
-		PropertyInheritanceInclusionOptions != PropertyInheritanceInclusionOptionsCache)
-	{
-		PropertyInheritanceInclusionOptionsCache = PropertyInheritanceInclusionOptions;
-		if (!bNeedToGenerateHeaders)
+		if (const bool bShouldIncludeUnsupportedProperties = FilterInstance->ShouldIncludeUnsupportedProperties(); 
+			bShouldIncludeUnsupportedProperties != bShouldIncludeUnsupportedPropertiesCache)
 		{
+			bShouldIncludeUnsupportedPropertiesCache = bShouldIncludeUnsupportedProperties;
 			bNeedToGenerateHeaders = true;
 		}
 	}
-	if (const bool bShouldIncludeUnsupportedProperties = SelectedFilter->ShouldIncludeUnsupportedProperties(); 
-		bShouldIncludeUnsupportedProperties != bShouldIncludeUnsupportedPropertiesCache)
+	else
 	{
-		bShouldIncludeUnsupportedPropertiesCache = bShouldIncludeUnsupportedProperties;
-		if (!bNeedToGenerateHeaders)
-		{
-			bNeedToGenerateHeaders = true;
-		}
+		UE_LOG(LogObjectMixerEditor, Display, TEXT("%hs: No Main UObjectMixerObjectFilter instance found."), __FUNCTION__);
 	}
 	
 	if (bNeedToGenerateHeaders)
diff --git a/Engine/Plugins/Editor/ObjectMixer/ObjectMixer/Source/ObjectMixer/Private/Views/List/SObjectMixerEditorListRow.cpp b/Engine/Plugins/Editor/ObjectMixer/ObjectMixer/Source/ObjectMixer/Private/Views/List/SObjectMixerEditorListRow.cpp
--- a/Engine/Plugins/Editor/ObjectMixer/ObjectMixer/Source/ObjectMixer/Private/Views/List/SObjectMixerEditorListRow.cpp
+++ b/Engine/Plugins/Editor/ObjectMixer/ObjectMixer/Source/ObjectMixer/Private/Views/List/SObjectMixerEditorListRow.cpp
@@ -92,7 +92,7 @@ public:
 
 			if (bRowObjectIsValid)
 			{
-				if (const UObjectMixerObjectFilter* Filter = InRow->GetObjectFilter())
+				if (const UObjectMixerObjectFilter* Filter = InRow->GetMainObjectFilterInstance())
 				{
 					TooltipText = Filter->GetRowTooltipText(RowObject, HybridChild.IsValid());
 				}
diff --git a/Engine/Plugins/Editor/ObjectMixer/ObjectMixer/Source/ObjectMixer/Private/Views/MainPanel/ObjectMixerEditorMainPanel.cpp b/Engine/Plugins/Editor/ObjectMixer/ObjectMixer/Source/ObjectMixer/Private/Views/MainPanel/ObjectMixerEditorMainPanel.cpp
--- a/Engine/Plugins/Editor/ObjectMixer/ObjectMixer/Source/ObjectMixer/Private/Views/MainPanel/ObjectMixerEditorMainPanel.cpp
+++ b/Engine/Plugins/Editor/ObjectMixer/ObjectMixer/Source/ObjectMixer/Private/Views/MainPanel/ObjectMixerEditorMainPanel.cpp
@@ -150,42 +150,64 @@ FString FObjectMixerEditorMainPanel::GetSearchStringFromSearchInputField() const
 	return MainPanelWidget->GetSearchStringFromSearchInputField();
 }
 
-void FObjectMixerEditorMainPanel::OnClassSelectionChanged(UClass* InNewClass)
+void FObjectMixerEditorMainPanel::SetDefaultFilterClass(UClass* InNewClass)
 {
-	SetObjectFilterClass(InNewClass);
+	DefaultFilterClass = InNewClass;
+	AddObjectFilterClass(InNewClass);
 }
 
-TObjectPtr<UClass> FObjectMixerEditorMainPanel::GetClassSelection() const
+bool FObjectMixerEditorMainPanel::IsClassSelected(UClass* InClass) const
 {
-	return GetObjectFilterClass();
+	return GetObjectFilterClasses().Contains(InClass);
 }
 
-bool FObjectMixerEditorMainPanel::IsClassSelected(UClass* InNewClass) const
+const TArray<TObjectPtr<UObjectMixerObjectFilter>>& FObjectMixerEditorMainPanel::GetObjectFilterInstances()
 {
-	return InNewClass == GetClassSelection();
+	if (ObjectFilterInstances.Num() == 0)
+	{
+		CacheObjectFilterObjects();
+	}
+
+	return ObjectFilterInstances;
+}
+
+const UObjectMixerObjectFilter* FObjectMixerEditorMainPanel::GetMainObjectFilterInstance()
+{
+	return GetObjectFilterInstances().Num() > 0 ? GetObjectFilterInstances()[0].Get() : nullptr;
 }
 
-UObjectMixerObjectFilter* FObjectMixerEditorMainPanel::GetObjectFilter()
+void FObjectMixerEditorMainPanel::CacheObjectFilterObjects()
 {
-	if (!ObjectFilterPtr.IsValid())
+	ObjectFilterInstances.Reset();
+	
+	for (const TSubclassOf<UObjectMixerObjectFilter> Class : GetObjectFilterClasses())
 	{
-		CacheObjectFilterObject();
+		UObjectMixerObjectFilter* NewInstance = NewObject<UObjectMixerObjectFilter>(GetTransientPackage(), Class);
+		ObjectFilterInstances.Add(TObjectPtr<UObjectMixerObjectFilter>(NewInstance));
 	}
-
-	return ObjectFilterPtr.Get();
 }
 
-void FObjectMixerEditorMainPanel::CacheObjectFilterObject()
+TSet<UClass*> FObjectMixerEditorMainPanel::GetObjectClassesToFilter()
 {
-	if (ObjectFilterPtr.IsValid())
+	TSet<UClass*> ReturnValue;
+	for (const TObjectPtr<UObjectMixerObjectFilter>& Filter : GetObjectFilterInstances())
 	{
-		ObjectFilterPtr.Reset();
+		ReturnValue.Append(Filter->GetObjectClassesToFilter());
 	}
-	
-	if (const UClass* Class = GetObjectFilterClass())
+		
+	return ReturnValue;
+}
+
+TSet<TSubclassOf<AActor>> FObjectMixerEditorMainPanel::GetObjectClassesToPlace()
+{
+	TSet<TSubclassOf<AActor>> ReturnValue;
+
+	for (const TObjectPtr<UObjectMixerObjectFilter>& Filter : GetObjectFilterInstances())
 	{
-		ObjectFilterPtr = TStrongObjectPtr(NewObject<UObjectMixerObjectFilter>(GetTransientPackage(), Class));
+		ReturnValue.Append(Filter->GetObjectClassesToPlace());
 	}
+		
+	return ReturnValue;
 }
 
 const TArray<TSharedRef<IObjectMixerEditorListFilter>>& FObjectMixerEditorMainPanel::GetListFilters() const
@@ -198,23 +220,44 @@ TArray<TWeakPtr<IObjectMixerEditorListFilter>> FObjectMixerEditorMainPanel::GetW
 	return MainPanelWidget->GetWeakActiveListFiltersSortedByName();
 }
 
-UObjectMixerEditorSerializedData* FObjectMixerEditorMainPanel::GetSerializedDataOutputtingFilterName(FName& OutFilterName) const
+void FObjectMixerEditorMainPanel::AddObjectFilterClass(UClass* InObjectFilterClass, const bool bCacheAndRebuild)
+{
+	if (ensureAlwaysMsgf(InObjectFilterClass->IsChildOf(UObjectMixerObjectFilter::StaticClass()), TEXT("%hs: Class '%s' is not a child of UObjectMixerObjectFilter."), __FUNCTION__, *InObjectFilterClass->GetName()))
+	{
+		ObjectFilterClasses.Add(InObjectFilterClass);
+
+		if (bCacheAndRebuild)
+		{
+			CacheAndRebuildFilters();
+		}
+	}
+}
+
+void FObjectMixerEditorMainPanel::RemoveObjectFilterClass(UClass* InObjectFilterClass, const bool bCacheAndRebuild)
 {
-	if (const TSubclassOf<UObjectMixerObjectFilter> Filter = GetObjectFilterClass())
+	if (ObjectFilterClasses.Remove(InObjectFilterClass) > 0 && bCacheAndRebuild)
 	{
-		OutFilterName = Filter->GetFName();
-		return GetMutableDefault<UObjectMixerEditorSerializedData>();
+		CacheAndRebuildFilters();
 	}
+}
 
-	return nullptr;
+UObjectMixerEditorSerializedData* FObjectMixerEditorMainPanel::GetSerializedData() const
+{
+	return GetMutableDefault<UObjectMixerEditorSerializedData>();
 }
 
 bool FObjectMixerEditorMainPanel::RequestAddObjectsToCollection(const FName& CollectionName, const TSet<FSoftObjectPath>& ObjectsToAdd) const
 {
-	FName FilterName;
-	if (UObjectMixerEditorSerializedData* SerializedData = GetSerializedDataOutputtingFilterName(FilterName))
+	if (UObjectMixerEditorSerializedData* SerializedData = GetSerializedData())
 	{
-		return SerializedData->AddObjectsToCollection(FilterName, CollectionName, ObjectsToAdd);
+		for (const TSubclassOf<UObjectMixerObjectFilter>& Class : GetObjectFilterClasses())
+		{
+			const FName FilterName = Class->GetFName();
+			if (SerializedData->DoesCollectionExist(FilterName, CollectionName))
+			{
+				return SerializedData->AddObjectsToCollection(FilterName, CollectionName, ObjectsToAdd);
+			}
+		}
 	}
 
 	return false;
@@ -222,10 +265,16 @@ bool FObjectMixerEditorMainPanel::RequestAddObjectsToCollection(const FName& Col
 
 bool FObjectMixerEditorMainPanel::RequestRemoveObjectsFromCollection(const FName& CollectionName, const TSet<FSoftObjectPath>& ObjectsToRemove) const
 {
-	FName FilterName;
-	if (UObjectMixerEditorSerializedData* SerializedData = GetSerializedDataOutputtingFilterName(FilterName))
+	if (UObjectMixerEditorSerializedData* SerializedData = GetSerializedData())
 	{
-		return SerializedData->RemoveObjectsFromCollection(FilterName, CollectionName, ObjectsToRemove);
+		for (const TSubclassOf<UObjectMixerObjectFilter>& Class : GetObjectFilterClasses())
+		{
+			const FName FilterName = Class->GetFName();
+			if (SerializedData->DoesCollectionExist(FilterName, CollectionName))
+			{
+				return SerializedData->RemoveObjectsFromCollection(FilterName, CollectionName, ObjectsToRemove);
+			}
+		}
 	}
 
 	return false;
@@ -233,10 +282,16 @@ bool FObjectMixerEditorMainPanel::RequestRemoveObjectsFromCollection(const FName
 
 bool FObjectMixerEditorMainPanel::RequestRemoveCollection(const FName& CollectionName) const
 {
-	FName FilterName;
-	if (UObjectMixerEditorSerializedData* SerializedData = GetSerializedDataOutputtingFilterName(FilterName))
+	if (UObjectMixerEditorSerializedData* SerializedData = GetSerializedData())
 	{
-		return SerializedData->RemoveCollection(FilterName, CollectionName);
+		for (const TSubclassOf<UObjectMixerObjectFilter>& Class : GetObjectFilterClasses())
+		{
+			const FName FilterName = Class->GetFName();
+			if (SerializedData->DoesCollectionExist(FilterName, CollectionName))
+			{
+				return SerializedData->RemoveCollection(FilterName, CollectionName);
+			}
+		}
 	}
 
 	return false;
@@ -244,10 +299,16 @@ bool FObjectMixerEditorMainPanel::RequestRemoveCollection(const FName& Collectio
 
 bool FObjectMixerEditorMainPanel::RequestDuplicateCollection(const FName& CollectionToDuplicateName, FName& DesiredDuplicateName) const
 {
-	FName FilterName;
-	if (UObjectMixerEditorSerializedData* SerializedData = GetSerializedDataOutputtingFilterName(FilterName))
+	if (UObjectMixerEditorSerializedData* SerializedData = GetSerializedData())
 	{
-		return SerializedData->DuplicateCollection(FilterName, CollectionToDuplicateName, DesiredDuplicateName);
+		for (const TSubclassOf<UObjectMixerObjectFilter>& Class : GetObjectFilterClasses())
+		{
+			const FName FilterName = Class->GetFName();
+			if (SerializedData->DoesCollectionExist(FilterName, CollectionToDuplicateName))
+			{
+				return SerializedData->DuplicateCollection(FilterName, CollectionToDuplicateName, DesiredDuplicateName);
+			}
+		}
 	}
 
 	return false;
@@ -255,10 +316,16 @@ bool FObjectMixerEditorMainPanel::RequestDuplicateCollection(const FName& Collec
 
 bool FObjectMixerEditorMainPanel::RequestReorderCollection(const FName& CollectionToMoveName, const FName& CollectionInsertBeforeName) const
 {
-	FName FilterName;
-	if (UObjectMixerEditorSerializedData* SerializedData = GetSerializedDataOutputtingFilterName(FilterName))
+	if (UObjectMixerEditorSerializedData* SerializedData = GetSerializedData())
 	{
-		return SerializedData->ReorderCollection(FilterName, CollectionToMoveName, CollectionInsertBeforeName);
+		for (const TSubclassOf<UObjectMixerObjectFilter>& Class : GetObjectFilterClasses())
+		{
+			const FName FilterName = Class->GetFName();
+			if (SerializedData->DoesCollectionExist(FilterName, CollectionToMoveName))
+			{
+				return SerializedData->ReorderCollection(FilterName, CollectionToMoveName, CollectionInsertBeforeName);
+			}
+		}
 	}
 
 	return false;
@@ -267,10 +334,16 @@ bool FObjectMixerEditorMainPanel::RequestReorderCollection(const FName& Collecti
 bool FObjectMixerEditorMainPanel::RequestRenameCollection(const FName& CollectionNameToRename,
 	const FName& NewCollectionName) const
 {
-	FName FilterName;
-	if (UObjectMixerEditorSerializedData* SerializedData = GetSerializedDataOutputtingFilterName(FilterName))
+	if (UObjectMixerEditorSerializedData* SerializedData = GetSerializedData())
 	{
-		return SerializedData->RenameCollection(FilterName, CollectionNameToRename, NewCollectionName);
+		for (const TSubclassOf<UObjectMixerObjectFilter>& Class : GetObjectFilterClasses())
+		{
+			const FName FilterName = Class->GetFName();
+			if (SerializedData->DoesCollectionExist(FilterName, CollectionNameToRename))
+			{
+				return SerializedData->RenameCollection(FilterName, CollectionNameToRename, NewCollectionName);
+			}
+		}
 	}
 
 	return false;
@@ -278,10 +351,16 @@ bool FObjectMixerEditorMainPanel::RequestRenameCollection(const FName& Collectio
 
 bool FObjectMixerEditorMainPanel::DoesCollectionExist(const FName& CollectionName) const
 {
-	FName FilterName;
-	if (UObjectMixerEditorSerializedData* SerializedData = GetSerializedDataOutputtingFilterName(FilterName))
+	if (UObjectMixerEditorSerializedData* SerializedData = GetSerializedData())
 	{
-		return SerializedData->DoesCollectionExist(FilterName, CollectionName);
+		for (const TSubclassOf<UObjectMixerObjectFilter>& Class : GetObjectFilterClasses())
+		{
+			const FName FilterName = Class->GetFName();
+			if (SerializedData->DoesCollectionExist(FilterName, CollectionName))
+			{
+				return true;
+			}
+		}
 	}
 
 	return false;
@@ -289,10 +368,16 @@ bool FObjectMixerEditorMainPanel::DoesCollectionExist(const FName& CollectionNam
 
 bool FObjectMixerEditorMainPanel::IsObjectInCollection(const FName& CollectionName, const FSoftObjectPath& InObject) const
 {
-	FName FilterName;
-	if (UObjectMixerEditorSerializedData* SerializedData = GetSerializedDataOutputtingFilterName(FilterName))
+	if (UObjectMixerEditorSerializedData* SerializedData = GetSerializedData())
 	{
-		return SerializedData->IsObjectInCollection(FilterName, CollectionName, InObject);
+		for (const TSubclassOf<UObjectMixerObjectFilter>& Class : GetObjectFilterClasses())
+		{
+			const FName FilterName = Class->GetFName();
+			if (SerializedData->DoesCollectionExist(FilterName, CollectionName))
+			{
+				return SerializedData->IsObjectInCollection(FilterName, CollectionName, InObject);
+			}
+		}
 	}
 
 	return false;
@@ -300,27 +385,40 @@ bool FObjectMixerEditorMainPanel::IsObjectInCollection(const FName& CollectionNa
 
 TSet<FName> FObjectMixerEditorMainPanel::GetCollectionsForObject(const FSoftObjectPath& InObject) const
 {
-	FName FilterName;
-	if (UObjectMixerEditorSerializedData* SerializedData = GetSerializedDataOutputtingFilterName(FilterName))
+	TSet<FName> ReturnValue;
+	if (UObjectMixerEditorSerializedData* SerializedData = GetSerializedData())
 	{
-		return SerializedData->GetCollectionsForObject(FilterName, InObject);
+		for (const TSubclassOf<UObjectMixerObjectFilter>& Class : GetObjectFilterClasses())
+		{
+			const FName FilterName = Class->GetFName();
+			ReturnValue.Append(SerializedData->GetCollectionsForObject(FilterName, InObject));
+		}
 	}
 
-	return {};
+	return ReturnValue;
 }
 
 TArray<FName> FObjectMixerEditorMainPanel::GetAllCollectionNames() const
 {
-	FName FilterName;
-	if (UObjectMixerEditorSerializedData* SerializedData = GetSerializedDataOutputtingFilterName(FilterName))
+	TArray<FName> ReturnValue;
+	if (UObjectMixerEditorSerializedData* SerializedData = GetSerializedData())
 	{
-		return SerializedData->GetAllCollectionNames(FilterName);
+		for (const TSubclassOf<UObjectMixerObjectFilter>& Class : GetObjectFilterClasses())
+		{
+			const FName FilterName = Class->GetFName();
+			ReturnValue.Append(SerializedData->GetAllCollectionNames(FilterName));
+		}
 	}
 
-	return {};
+	return ReturnValue;
 }
 
 TSet<TSharedRef<FObjectMixerEditorListFilter_Collection>> FObjectMixerEditorMainPanel::GetCurrentCollectionSelection() const
 {
 	return MainPanelWidget->GetCurrentCollectionSelection();
 }
+
+const TSubclassOf<UObjectMixerObjectFilter>& FObjectMixerEditorMainPanel::GetDefaultFilterClass() const
+{
+	return DefaultFilterClass;
+}
diff --git a/Engine/Plugins/Editor/ObjectMixer/ObjectMixer/Source/ObjectMixer/Private/Views/MainPanel/SObjectMixerEditorMainPanel.cpp b/Engine/Plugins/Editor/ObjectMixer/ObjectMixer/Source/ObjectMixer/Private/Views/MainPanel/SObjectMixerEditorMainPanel.cpp
--- a/Engine/Plugins/Editor/ObjectMixer/ObjectMixer/Source/ObjectMixer/Private/Views/MainPanel/SObjectMixerEditorMainPanel.cpp
+++ b/Engine/Plugins/Editor/ObjectMixer/ObjectMixer/Source/ObjectMixer/Private/Views/MainPanel/SObjectMixerEditorMainPanel.cpp
@@ -2,6 +2,7 @@
 
 #include "Views/MainPanel/SObjectMixerEditorMainPanel.h"
 
+#include "ContentBrowserModule.h"
 #include "Engine/Blueprint.h"
 #include "ObjectMixerEditorLog.h"
 #include "ObjectMixerEditorModule.h"
@@ -17,6 +18,7 @@
 #include "AssetRegistry/IAssetRegistry.h"
 #include "Framework/Application/SlateApplication.h"
 #include "EditorActorFolders.h"
+#include "IContentBrowserSingleton.h"
 #include "Framework/MultiBox/MultiBoxBuilder.h"
 #include "IPlacementModeModule.h"
 #include "Kismet2/SClassPickerDialog.h"
@@ -29,6 +31,7 @@
 #include "Widgets/Input/SSearchBox.h"
 #include "Widgets/Layout/SWidgetSwitcher.h"
 #include "Widgets/Layout/SWrapBox.h"
+#include "Widgets/SBoxPanel.h"
 #include "Widgets/Text/STextBlock.h"
 
 #define LOCTEXT_NAMESPACE "ObjectMixerEditor"
@@ -214,15 +217,23 @@ TSharedRef<SWidget> SObjectMixerEditorMainPanel::GenerateToolbar()
 
 TSharedRef<SWidget> SObjectMixerEditorMainPanel::OnGenerateAddObjectButtonMenu() const
 {
-	const TSet<TSubclassOf<AActor>> SubclassesOfActor = MainPanelModel.Pin()->GetObjectClassesToPlace();
+	TSet<UClass*> ClassesToPlace;
+	for (const TObjectPtr<UObjectMixerObjectFilter>& Instance : MainPanelModel.Pin()->GetObjectFilterInstances())
+	{
+		const TSet<TSubclassOf<AActor>> SubclassesOfActor = Instance->GetObjectClassesToPlace();
+		if (SubclassesOfActor.Num() > 0)
+		{
+			ClassesToPlace.Append(
+				Instance->GetParentAndChildClassesFromSpecifiedClasses(
+					SubclassesOfActor,
+					Instance->GetObjectMixerPlacementClassInclusionOptions()
+				)
+			);
+		}
+	}
 
-	if (SubclassesOfActor.Num() > 0)
+	if (ClassesToPlace.Num() > 0)
 	{
-		TSet<UClass*> ClassesToPlace =
-		   MainPanelModel.Pin()->GetObjectFilter()->GetParentAndChildClassesFromSpecifiedClasses(
-			   SubclassesOfActor,
-			   MainPanelModel.Pin()->GetObjectFilter()->GetObjectMixerPlacementClassInclusionOptions());
-	
 		FMenuBuilder AddObjectButtonMenuBuilder = FMenuBuilder(true, nullptr);
 
 		for (const UClass* Class : ClassesToPlace)
@@ -263,9 +274,27 @@ bool IsBlueprintFilter(const FAssetData& BlueprintClassData)
 	return bInheritsFromBlueprintFilter;
 }
 
+struct FAssetClassMap
+{
+	TObjectPtr<UClass> Class = nullptr;
+	FAssetData AssetData;
+
+	bool operator==(const FAssetClassMap& Other) const
+	{
+		return Class == Other.Class;
+	}
+};
+
 TSharedRef<SWidget> SObjectMixerEditorMainPanel::OnGenerateFilterClassMenu()
 {
-	FMenuBuilder MenuBuilder(true, nullptr);
+	TSharedRef<SVerticalBox> VerticalBox = SNew(SVerticalBox);
+
+	TSharedRef<SBox> OuterBox =
+		SNew(SBox)
+		.Padding(8)
+		[
+			VerticalBox
+		];
 
 	// Get C++ Derivatives (and maybe Blueprint derivatives)
 	TArray<UClass*> DerivedClasses;
@@ -274,6 +303,13 @@ TSharedRef<SWidget> SObjectMixerEditorMainPanel::OnGenerateFilterClassMenu()
 	DerivedClasses.Remove(UObjectMixerObjectFilter::StaticClass());
 	DerivedClasses.Remove(UObjectMixerBlueprintObjectFilter::StaticClass());
 
+	TArray<FAssetClassMap> AssetClassMaps;
+
+	for (UClass* Class : DerivedClasses)
+	{
+		AssetClassMaps.Add({Class});
+	}
+
 	// Get remaining Blueprint derivatives
 	FAssetRegistryModule& AssetRegistryModule = FModuleManager::LoadModuleChecked< FAssetRegistryModule >(FName("AssetRegistry"));
 	IAssetRegistry& AssetRegistry = AssetRegistryModule.Get();
@@ -284,75 +320,121 @@ TSharedRef<SWidget> SObjectMixerEditorMainPanel::OnGenerateFilterClassMenu()
 	{
 		if (IsBlueprintFilter(Asset))
 		{
-			if (UBlueprint* BlueprintAsset = Cast<UBlueprint>(Asset.GetAsset()))
+			if (const UBlueprint* BlueprintAsset = Cast<UBlueprint>(Asset.GetAsset()))
 			{
 				UClass* LoadedClass = BlueprintAsset->GeneratedClass;
 				if (ensure(LoadedClass && BlueprintAsset->ParentClass))
 				{
-					DerivedClasses.AddUnique(LoadedClass);
+					if (FAssetClassMap* Match = Algo::FindByPredicate(
+						AssetClassMaps,
+						[LoadedClass](const FAssetClassMap& ClassMap)
+						{
+							return ClassMap.Class == LoadedClass;
+						}))
+					{
+						Match->AssetData = Asset;
+					}
+					else
+					{
+						AssetClassMaps.Add({LoadedClass, Asset});
+					}
 				}
 			}
 		}
 	}
 
-	DerivedClasses.Sort([](UClass& A, UClass& B)
-	{
-		return A.GetFName().LexicalLess(B.GetFName());
-	});
-
-	if (DerivedClasses.Num())
+	if (AssetClassMaps.Num())
 	{
 		check(MainPanelModel.Pin());
+
+		AssetClassMaps.Sort([](const FAssetClassMap& A, const FAssetClassMap& B)
+		{
+			return A.Class.GetFName().LexicalLess(B.Class.GetFName());
+		});
 		
-		MenuBuilder.BeginSection(NAME_None, LOCTEXT("SelectClassMenuSection", "Select Class"));
+		FilterClassSelectionInfos.Empty(FilterClassSelectionInfos.Num());
+		for (const FAssetClassMap& AssetClassMap : AssetClassMaps)
 		{
-			for (UClass* DerivedClass : DerivedClasses)
+			if (IsValid(AssetClassMap.Class))
 			{
-				if (IsValid(DerivedClass))
+				if (AssetClassMap.Class->GetName().StartsWith(TEXT("SKEL_")) || AssetClassMap.Class->GetName().StartsWith(TEXT("REINST_")))
 				{
-					if (DerivedClass->GetName().StartsWith(TEXT("SKEL_")) || DerivedClass->GetName().StartsWith(TEXT("REINST_")))
-					{
-						continue;
-					}
+					continue;
+				}
 
-					if (DerivedClass->HasAnyClassFlags(CLASS_Abstract | CLASS_HideDropDown | CLASS_Deprecated))
-					{
-						continue;
-					}
-					
-					MenuBuilder.AddMenuEntry(
-					   FText::FromString(DerivedClass->GetName().EndsWith(TEXT("_C")) ? DerivedClass->GetName().LeftChop(2) : DerivedClass->GetName()),
-					   FText::FromString(DerivedClass->GetClassPathName().ToString()),
-					   FSlateIcon(),
-					   FUIAction(
-						   FExecuteAction::CreateSP(MainPanelModel.Pin().ToSharedRef(), &FObjectMixerEditorMainPanel::SetObjectFilterClass, DerivedClass),
-						   FCanExecuteAction::CreateLambda([](){ return true; }),
-						   FIsActionChecked::CreateSP(MainPanelModel.Pin().ToSharedRef(), &FObjectMixerEditorMainPanel::IsClassSelected, DerivedClass)
-					   ),
-					   NAME_None,
-					   EUserInterfaceActionType::RadioButton
-				   );
+				if (AssetClassMap.Class->HasAnyClassFlags(CLASS_Abstract | CLASS_HideDropDown | CLASS_Deprecated))
+				{
+					continue;
 				}
+
+				const bool bIsDefaultClass = MainPanelModel.Pin()->GetDefaultFilterClass() == AssetClassMap.Class;
+				
+				const FText TooltipText = bIsDefaultClass ?
+					FText::Format(
+						LOCTEXT("DefaultClassDisclaimer","This class explicitly cannot be disabled in {0}"), 
+						FText::FromName(MainPanelModel.Pin()->GetModuleName())) :
+					FText::FromString(AssetClassMap.Class->GetClassPathName().ToString()
+				);
+
+				FilterClassSelectionInfos.Add({AssetClassMap.Class, MainPanelModel.Pin()->IsClassSelected(AssetClassMap.Class)});
+
+				VerticalBox->AddSlot()
+				.Padding(FMargin(0, 0, 0, 8))
+				.AutoHeight()
+				[
+					SNew(SFilterClassMenuItem, AssetClassMap, bIsDefaultClass, FilterClassSelectionInfos, TooltipText)
+				];
 			}
 		}
-		MenuBuilder.EndSection();
+
+		VerticalBox->AddSlot()
+		.HAlign(HAlign_Fill)
+		[
+			SNew(SButton)
+			.Text(LOCTEXT("SelectFilterClassMenu_ApplyButton", "Apply"))
+			.HAlign(HAlign_Center)
+			.OnClicked(FOnClicked::CreateLambda([this]()
+			   {
+					if (const TSharedPtr<FObjectMixerEditorMainPanel> PinnedMainPanel = MainPanelModel.Pin())
+					{
+						PinnedMainPanel->ResetObjectFilterClasses(false);
+						for (const FFilterClassSelectionInfo& Info : FilterClassSelectionInfos)
+						{
+							if (Info.bIsUserSelected)
+							{
+								PinnedMainPanel->AddObjectFilterClass(Info.Class, false);
+							}
+						}
+
+						PinnedMainPanel->CacheAndRebuildFilters();
+					}
+
+					return FReply::Handled();
+			   })
+			)
+		];
 	}
 	else
 	{
-		MenuBuilder.AddMenuEntry(LOCTEXT("NoFilterClassesAvailable", "No filter classes available."), FText::GetEmpty(), FSlateIcon(), FUIAction());
+		VerticalBox->AddSlot()
+		.HAlign(HAlign_Center)
+		[
+			SNew(STextBlock)
+			.Text(LOCTEXT("NoFilterClassesAvailable", "No filter classes available."))
+		];
 	}
 
-	TSharedRef<SWidget> Widget = MenuBuilder.MakeWidget();
-	FChildren* ChildWidgets = Widget->GetChildren();
+	FChildren* ChildWidgets = VerticalBox->GetChildren();
 	for (int32 ChildItr = 0; ChildItr < ChildWidgets->Num(); ChildItr++)
 	{
 		const TSharedRef<SWidget>& Child = ChildWidgets->GetChildAt(ChildItr);
 
 		Child->EnableToolTipForceField(false);
 	}
-	Widget->EnableToolTipForceField(false);
+	VerticalBox->EnableToolTipForceField(false);
+	OuterBox->EnableToolTipForceField(false);
 	
-	return Widget;
+	return OuterBox;
 }
 
 TSharedRef<SWidget> SObjectMixerEditorMainPanel::BuildShowOptionsMenu()
@@ -360,49 +442,46 @@ TSharedRef<SWidget> SObjectMixerEditorMainPanel::BuildShowOptionsMenu()
 	FMenuBuilder ShowOptionsMenuBuilder = FMenuBuilder(true, nullptr);
 
 	// No need to select filter class from outside generic instance
-	if (MainPanelModel.Pin()->GetModuleName() == FObjectMixerEditorModule::BaseObjectMixerModuleName)
+	ShowOptionsMenuBuilder.BeginSection("ListViewOptions", LOCTEXT("FilterClassManagementSection", "Filter Class Management"));
 	{
-		ShowOptionsMenuBuilder.BeginSection("ListViewOptions", LOCTEXT("FilterClassManagementSection", "Filter Class Management"));
-		{
-			// Filter Class Management Button
-			const TSharedRef<SWidget> FilterClassManagementButton =
-				SNew(SBox)
-				.Padding(8, 0)
+		// Filter Class Management Button
+		const TSharedRef<SWidget> FilterClassManagementButton =
+			SNew(SBox)
+			.Padding(8, 0)
+			[
+				SNew(SComboButton)
+				.ToolTipText(LOCTEXT("FilterClassManagementButton_Tooltip", "Select a filter class"))
+				.ContentPadding(FMargin(4, 0.5f))
+				.ComboButtonStyle(&FAppStyle::Get().GetWidgetStyle<FComboButtonStyle>("ComboButton"))
+				.OnGetMenuContent(this, &SObjectMixerEditorMainPanel::OnGenerateFilterClassMenu)
+				.ForegroundColor(FStyleColors::Foreground)
+				.MenuPlacement(EMenuPlacement::MenuPlacement_MenuRight)
+				.ButtonContent()
 				[
-					SNew(SComboButton)
-					.ToolTipText(LOCTEXT("FilterClassManagementButton_Tooltip", "Select a filter class"))
-					.ContentPadding(FMargin(4, 0.5f))
-					.ComboButtonStyle(&FAppStyle::Get().GetWidgetStyle<FComboButtonStyle>("ComboButton"))
-					.OnGetMenuContent(this, &SObjectMixerEditorMainPanel::OnGenerateFilterClassMenu)
-					.ForegroundColor(FStyleColors::Foreground)
-					.MenuPlacement(EMenuPlacement::MenuPlacement_MenuRight)
-					.ButtonContent()
+					SNew(SHorizontalBox)
+
+					+ SHorizontalBox::Slot()
+					.Padding(0, 1, 4, 0)
+					.AutoWidth()
 					[
-						SNew(SHorizontalBox)
-
-						+ SHorizontalBox::Slot()
-						.Padding(0, 1, 4, 0)
-						.AutoWidth()
-						[
-							SNew(SImage)
-							.Image(FAppStyle::Get().GetBrush("Icons.Filter"))
-							.ColorAndOpacity(FSlateColor::UseForeground())
-						]
-
-						+ SHorizontalBox::Slot()
-						.Padding(0, 1, 0, 0)
-						.AutoWidth()
-						[
-							SNew(STextBlock)
-							.Text(LOCTEXT("FilterClassToolbarButton", "Object Filter Class"))
-						]
+						SNew(SImage)
+						.Image(FAppStyle::Get().GetBrush("Icons.Filter"))
+						.ColorAndOpacity(FSlateColor::UseForeground())
 					]
-				];
 
-			ShowOptionsMenuBuilder.AddWidget(FilterClassManagementButton, FText::GetEmpty());
-		}
-		ShowOptionsMenuBuilder.EndSection();
+					+ SHorizontalBox::Slot()
+					.Padding(0, 1, 0, 0)
+					.AutoWidth()
+					[
+						SNew(STextBlock)
+						.Text(LOCTEXT("FilterClassToolbarButton", "Object Filter Class"))
+					]
+				]
+			];
+
+		ShowOptionsMenuBuilder.AddWidget(FilterClassManagementButton, FText::GetEmpty());
 	}
+	ShowOptionsMenuBuilder.EndSection();
 
 	// Add List View Mode Options
 	ShowOptionsMenuBuilder.BeginSection("ListViewOptions", LOCTEXT("ListViewOptionsSection", "List View Options"));
@@ -553,6 +632,88 @@ void SObjectMixerEditorMainPanel::SetSingleCollectionSelection(const FName& Coll
 	}
 }
 
+void SFilterClassMenuItem::Construct(const FArguments& InArgs, FAssetClassMap AssetClassMap, const bool bIsDefaultClass,
+	TArray<SObjectMixerEditorMainPanel::FFilterClassSelectionInfo>& FilterClassSelectionInfos, const FText TooltipText)
+{
+	const bool bHasValidAssetData = AssetClassMap.AssetData.IsValid();
+	ChildSlot
+	[
+		SNew(SHorizontalBox)
+		.ToolTipText(TooltipText)
+		.IsEnabled(!bIsDefaultClass)
+
+		+SHorizontalBox::Slot()
+		.Padding(FMargin(0, 0, 8, 0))
+		.AutoWidth()
+		[
+			SNew(SCheckBox)
+			.OnCheckStateChanged_Lambda([this, AssetClassMap, &FilterClassSelectionInfos](ECheckBoxState NewState)
+		   {
+				if (SObjectMixerEditorMainPanel::FFilterClassSelectionInfo* Match =
+					Algo::FindByPredicate(FilterClassSelectionInfos,
+						[AssetClassMap](const SObjectMixerEditorMainPanel::FFilterClassSelectionInfo& Other)
+						{
+							return Other.Class == AssetClassMap.Class;
+						}))
+				{
+					   Match->bIsUserSelected = !Match->bIsUserSelected;
+				}
+		   })
+		   .IsChecked_Lambda([this, AssetClassMap, &FilterClassSelectionInfos]()
+		   {
+		   		bool bShouldBeChecked = false;
+		   	
+				if (const SObjectMixerEditorMainPanel::FFilterClassSelectionInfo* Match =
+			   		Algo::FindByPredicate(FilterClassSelectionInfos,
+						[AssetClassMap](const SObjectMixerEditorMainPanel::FFilterClassSelectionInfo& Other)
+						{
+							return Other.Class == AssetClassMap.Class;
+						}))
+			   {
+				   bShouldBeChecked = Match->bIsUserSelected;
+			   }
+			   return bShouldBeChecked ? ECheckBoxState::Checked : ECheckBoxState::Unchecked;
+		   })
+		]
+
+		+SHorizontalBox::Slot()
+		.HAlign(HAlign_Fill)
+		.VAlign(VAlign_Center)
+		.Padding(FMargin(0, 0, 8, 0))
+		[
+			SNew(STextBlock)
+			.Text(FText::FromString(
+				AssetClassMap.Class->GetName().EndsWith(TEXT("_C")) ? AssetClassMap.Class->GetName().LeftChop(2) : AssetClassMap.Class->GetName()))
+		]
+
+		+SHorizontalBox::Slot()
+		.HAlign(HAlign_Right)
+		.AutoWidth()	
+		[
+			SNew(SButton)
+			.ButtonStyle(FAppStyle::Get(), "SimpleButton")
+			.HAlign(HAlign_Center)
+			.VAlign(VAlign_Center)
+			.ContentPadding(FMargin(2))
+			.IsEnabled(bHasValidAssetData)
+			.ToolTipText(bHasValidAssetData ?
+				LOCTEXT("BrowseTooltip", "Browses to the associated asset and selects it in the most recently used Content Browser (summoning one if necessary)")
+					: LOCTEXT("NoBlueprintFilterFound", "This filter class is not a Blueprint class."))
+			.OnClicked_Lambda([AssetClassMap]()
+			{
+				const FContentBrowserModule& ContentBrowserModule = FModuleManager::Get().LoadModuleChecked<FContentBrowserModule>("ContentBrowser");
+				ContentBrowserModule.Get().SyncBrowserToAssets({AssetClassMap.AssetData});
+
+				return FReply::Handled();
+			})
+			[
+				SNew(SImage)
+				.Image(FAppStyle::Get().GetBrush("SystemWideCommands.FindInContentBrowser.Small"))
+			]
+		]
+	];
+}
+
 FText SObjectMixerEditorMainPanel::GetSearchTextFromSearchInputField() const
 {
 	return ensureAlwaysMsgf(SearchBoxPtr.IsValid(),
diff --git a/Engine/Plugins/Editor/ObjectMixer/ObjectMixer/Source/ObjectMixer/Public/ObjectMixerEditorModule.h b/Engine/Plugins/Editor/ObjectMixer/ObjectMixer/Source/ObjectMixer/Public/ObjectMixerEditorModule.h
--- a/Engine/Plugins/Editor/ObjectMixer/ObjectMixer/Source/ObjectMixer/Public/ObjectMixerEditorModule.h
+++ b/Engine/Plugins/Editor/ObjectMixer/ObjectMixer/Source/ObjectMixer/Public/ObjectMixerEditorModule.h
@@ -66,6 +66,12 @@ public:
 
 	TSharedPtr<FWorkspaceItem> GetWorkspaceGroup();
 
+	/**
+	 * This is the filter class used to initialize the MainPanel.
+	 * This filter class cannot be turned off by the end user.
+	 */
+	const TSubclassOf<UObjectMixerObjectFilter>& GetDefaultFilterClass() const;
+
 	const static FName BaseObjectMixerModuleName;
 
 protected:
@@ -86,7 +92,10 @@ protected:
 	FText MenuItemTooltip;
 	ETabSpawnerMenuType::Type TabSpawnerType = ETabSpawnerMenuType::Enabled;
 
-	// If set, this is the filter class used to initialize the MainPanel
+	/**
+	 * If set, this is the filter class used to initialize the MainPanel.
+	 * This filter class cannot be turned off by the end user.
+	 */
 	TSubclassOf<UObjectMixerObjectFilter> DefaultFilterClass;
 	
 	TSet<FDelegateHandle> DelegateHandles;
diff --git a/Engine/Plugins/Editor/ObjectMixer/ObjectMixer/Source/ObjectMixer/Public/Views/List/ObjectMixerEditorListRow.h b/Engine/Plugins/Editor/ObjectMixer/ObjectMixer/Source/ObjectMixer/Public/Views/List/ObjectMixerEditorListRow.h
--- a/Engine/Plugins/Editor/ObjectMixer/ObjectMixer/Source/ObjectMixer/Public/Views/List/ObjectMixerEditorListRow.h
+++ b/Engine/Plugins/Editor/ObjectMixer/ObjectMixer/Source/ObjectMixer/Public/Views/List/ObjectMixerEditorListRow.h
@@ -14,6 +14,7 @@
 #include "PropertyHandle.h"
 #include "Templates/SharedPointer.h"
 #include "UObject/Object.h"
+#include "UObject/StrongObjectPtr.h"
 
 class SObjectMixerEditorList;
 class UObjectMixerObjectFilter;
@@ -179,7 +180,9 @@ struct OBJECTMIXEREDITOR_API FObjectMixerEditorListRow : TSharedFromThis<FObject
 		return FolderRef.GetPath();
 	}
 
-	UObjectMixerObjectFilter* GetObjectFilter() const;
+	const TArray<TObjectPtr<UObjectMixerObjectFilter>>& GetObjectFilterInstances() const;
+
+	const UObjectMixerObjectFilter* GetMainObjectFilterInstance() const;
 
 	bool IsObjectRefInCollection(const FName& CollectionName) const;
 
diff --git a/Engine/Plugins/Editor/ObjectMixer/ObjectMixer/Source/ObjectMixer/Public/Views/MainPanel/ObjectMixerEditorMainPanel.h b/Engine/Plugins/Editor/ObjectMixer/ObjectMixer/Source/ObjectMixer/Public/Views/MainPanel/ObjectMixerEditorMainPanel.h
--- a/Engine/Plugins/Editor/ObjectMixer/ObjectMixer/Source/ObjectMixer/Public/Views/MainPanel/ObjectMixerEditorMainPanel.h
+++ b/Engine/Plugins/Editor/ObjectMixer/ObjectMixer/Source/ObjectMixer/Public/Views/MainPanel/ObjectMixerEditorMainPanel.h
@@ -4,8 +4,9 @@
 
 #include "ObjectFilter/ObjectMixerEditorObjectFilter.h"
 
+#include "ObjectMixerEditorModule.h"
+
 #include "Folder.h"
-#include "UObject/StrongObjectPtr.h"
 #include "Widgets/SWidget.h"
 
 class FObjectMixerEditorList;
@@ -21,7 +22,7 @@ struct FObjectMixerEditorListRow;
 DECLARE_MULTICAST_DELEGATE(FOnPreFilterChange)
 DECLARE_MULTICAST_DELEGATE(FOnPostFilterChange)
 
-class OBJECTMIXEREDITOR_API FObjectMixerEditorMainPanel : public TSharedFromThis<FObjectMixerEditorMainPanel>
+class OBJECTMIXEREDITOR_API FObjectMixerEditorMainPanel : public TSharedFromThis<FObjectMixerEditorMainPanel>, public FGCObject
 {
 public:
 	FObjectMixerEditorMainPanel(const FName InModuleName)
@@ -68,13 +69,14 @@ public:
 	FText GetSearchTextFromSearchInputField() const;
 	FString GetSearchStringFromSearchInputField() const;
 
-	void OnClassSelectionChanged(UClass* InNewClass);
-	TObjectPtr<UClass> GetClassSelection() const;
-	bool IsClassSelected(UClass* InNewClass) const;
+	void SetDefaultFilterClass(UClass* InNewClass);
+	bool IsClassSelected(UClass* InClass) const;
+
+	const TArray<TObjectPtr<UObjectMixerObjectFilter>>& GetObjectFilterInstances();
 
-	UObjectMixerObjectFilter* GetObjectFilter();
+	const UObjectMixerObjectFilter* GetMainObjectFilterInstance();
 
-	void CacheObjectFilterObject();
+	void CacheObjectFilterObjects();
 
 	/**
 	 * Get the style of the tree (flat list or hierarchy)
@@ -95,49 +97,41 @@ public:
 	/**
 	 * Returns result from Filter->GetObjectClassesToFilter.
 	 */
-	TSet<UClass*> GetObjectClassesToFilter()
-	{
-		if (const UObjectMixerObjectFilter* Filter = GetObjectFilter())
-		{
-			return Filter->GetObjectClassesToFilter();
-		}
-		
-		return {};
-	}
+	TSet<UClass*> GetObjectClassesToFilter();
 
 	/**
 	 * Returns result from Filter->GetObjectClassesToPlace.
 	 */
-	TSet<TSubclassOf<AActor>> GetObjectClassesToPlace()
-	{
-		TSet<TSubclassOf<AActor>> ReturnValue;
-
-		if (const UObjectMixerObjectFilter* Filter = GetObjectFilter())
-		{
-			ReturnValue = Filter->GetObjectClassesToPlace();
-		}
-		
-		return ReturnValue;
-	}
+	TSet<TSubclassOf<AActor>> GetObjectClassesToPlace();
 
 	const TArray<TSharedRef<IObjectMixerEditorListFilter>>& GetListFilters() const;
 	TArray<TWeakPtr<IObjectMixerEditorListFilter>> GetWeakActiveListFiltersSortedByName() const;
 
-	TSubclassOf<UObjectMixerObjectFilter> GetObjectFilterClass() const
+	const TArray<TSubclassOf<UObjectMixerObjectFilter>>& GetObjectFilterClasses() const
 	{
-		return ObjectFilterClass;
+		return ObjectFilterClasses;
 	}
 
-	void SetObjectFilterClass(UClass* InObjectFilterClass)
+	void AddObjectFilterClass(UClass* InObjectFilterClass, const bool bCacheAndRebuild = true);
+
+	void RemoveObjectFilterClass(UClass* InObjectFilterClass, const bool bCacheAndRebuild = true);
+
+	void ResetObjectFilterClasses(const bool bCacheAndRebuild = true)
 	{
-		if (ensureAlwaysMsgf(InObjectFilterClass->IsChildOf(UObjectMixerObjectFilter::StaticClass()), TEXT("%hs: Class '%s' is not a child of UObjectMixerObjectFilter."), __FUNCTION__, *InObjectFilterClass->GetName()))
+		ObjectFilterClasses.Empty(ObjectFilterClasses.Num());
+
+		if (bCacheAndRebuild)
 		{
-			ObjectFilterClass = InObjectFilterClass;
-			CacheObjectFilterObject();
-			RequestRebuildList();
+			CacheAndRebuildFilters();
 		}
 	}
 
+	void CacheAndRebuildFilters()
+	{
+		CacheObjectFilterObjects();
+		RequestRebuildList();
+	}
+
 	FName GetModuleName() const
 	{
 		return ModuleName;
@@ -148,7 +142,7 @@ public:
 	/**
 	 * Get a pointer to the UObjectMixerEditorSerializedData object along with the name of the filter represented by this MainPanel instance.
 	 */
-	UObjectMixerEditorSerializedData* GetSerializedDataOutputtingFilterName(FName& OutFilterName) const;
+	UObjectMixerEditorSerializedData* GetSerializedData() const;
 	bool RequestAddObjectsToCollection(const FName& CollectionName, const TSet<FSoftObjectPath>& ObjectsToAdd) const;
 	bool RequestRemoveObjectsFromCollection(const FName& CollectionName, const TSet<FSoftObjectPath>& ObjectsToRemove) const;
 	bool RequestRemoveCollection(const FName& CollectionName) const;
@@ -165,24 +159,52 @@ public:
 	 */
 	TSet<TSharedRef<FObjectMixerEditorListFilter_Collection>> GetCurrentCollectionSelection() const;
 
+	/**
+	 * This is the filter class used to initialize the MainPanel.
+	 * This filter class cannot be turned off by the end user.
+	 */
+	const TSubclassOf<UObjectMixerObjectFilter>& GetDefaultFilterClass() const;
+
 	FOnPreFilterChange OnPreFilterChange;
 	FOnPostFilterChange OnPostFilterChange;
 
 	TSharedPtr<FUICommandList> ObjectMixerElementEditCommands;
 	TSharedPtr<FUICommandList> ObjectMixerFolderEditCommands;
 
+protected:
+
+	virtual void AddReferencedObjects( FReferenceCollector& Collector )  override
+	{
+		Collector.AddReferencedObjects(ObjectFilterInstances);
+	}
+	virtual FString GetReferencerName() const override
+	{
+		return TEXT("FObjectMixerEditorMainPanel");
+	}
+
 private:
 
 	TSharedPtr<SObjectMixerEditorMainPanel> MainPanelWidget;
 
 	TSharedPtr<FObjectMixerEditorList> EditorListModel;
 
-	TStrongObjectPtr<UObjectMixerObjectFilter> ObjectFilterPtr;
+	TArray<TObjectPtr<UObjectMixerObjectFilter>> ObjectFilterInstances;
+
+	/**
+	 * The classes used to generate property edit columns
+	 */
+	TArray<TSubclassOf<UObjectMixerObjectFilter>> ObjectFilterClasses;
+
+	/**
+	 * The FObjectMixerEditorModule subclass that created this instance.
+	 */
+	TWeakPtr<FObjectMixerEditorModule> SpawningModulePtr;
 
 	/**
-	 * The class used to generate property edit columns
+	 * If set, this is the filter class used to initialize the MainPanel.
+	 * This filter class cannot be turned off by the end user.
 	 */
-	TSubclassOf<UObjectMixerObjectFilter> ObjectFilterClass;
+	TSubclassOf<UObjectMixerObjectFilter> DefaultFilterClass;
 
 	/**
 	 * Determines the style of the tree (flat list or hierarchy)
diff --git a/Engine/Plugins/Editor/ObjectMixer/ObjectMixer/Source/ObjectMixer/Public/Views/MainPanel/SObjectMixerEditorMainPanel.h b/Engine/Plugins/Editor/ObjectMixer/ObjectMixer/Source/ObjectMixer/Public/Views/MainPanel/SObjectMixerEditorMainPanel.h
--- a/Engine/Plugins/Editor/ObjectMixer/ObjectMixer/Source/ObjectMixer/Public/Views/MainPanel/SObjectMixerEditorMainPanel.h
+++ b/Engine/Plugins/Editor/ObjectMixer/ObjectMixer/Source/ObjectMixer/Public/Views/MainPanel/SObjectMixerEditorMainPanel.h
@@ -12,8 +12,13 @@ class FObjectMixerEditorMainPanel;
 
 class OBJECTMIXEREDITOR_API SObjectMixerEditorMainPanel final : public SCompoundWidget
 {
-
 public:
+	
+	struct FFilterClassSelectionInfo
+	{
+		UClass* Class;
+		bool bIsUserSelected;
+	};
 
 	SLATE_BEGIN_ARGS(SObjectMixerEditorMainPanel)
 	{}
@@ -81,7 +86,9 @@ private:
 	TSharedRef<SWidget> GenerateToolbar();
 	TSharedRef<SWidget> OnGenerateAddObjectButtonMenu() const;
 
+	TArray<FFilterClassSelectionInfo> FilterClassSelectionInfos;
 	TSharedRef<SWidget> OnGenerateFilterClassMenu();
+	
 	TSharedRef<SWidget> BuildShowOptionsMenu();
 
 	void OnSearchTextChanged(const FText& Text);
@@ -96,3 +103,18 @@ private:
 	/** Disable all collection filters except CollectionToEnableName */
 	void SetSingleCollectionSelection(const FName& CollectionToEnableName = UObjectMixerEditorSerializedData::AllCollectionName);
 };
+
+
+
+class SFilterClassMenuItem : public SCompoundWidget
+{
+public:
+
+	SLATE_BEGIN_ARGS(SObjectMixerEditorMainPanel)
+	{}
+
+	SLATE_END_ARGS()
+
+	void Construct(const FArguments& InArgs, struct FAssetClassMap AssetClassMap, const bool bIsDefaultClass,
+		TArray<SObjectMixerEditorMainPanel::FFilterClassSelectionInfo>& FilterClassSelectionInfos, const FText TooltipText);
+};
