Object Mixer:

-Fix: //UE5/Main - Compile UnrealEditor Win64, VS2022 Latest Toolchain C++20 - Error: no overloaded function could convert all the argument types in SObjectMixerEditorListRow.cpp

diff --git a/Engine/Plugins/Editor/ObjectMixer/ObjectMixer/Source/ObjectMixer/Private/Views/List/SObjectMixerEditorListRow.cpp b/Engine/Plugins/Editor/ObjectMixer/ObjectMixer/Source/ObjectMixer/Private/Views/List/SObjectMixerEditorListRow.cpp
--- a/Engine/Plugins/Editor/ObjectMixer/ObjectMixer/Source/ObjectMixer/Private/Views/List/SObjectMixerEditorListRow.cpp
+++ b/Engine/Plugins/Editor/ObjectMixer/ObjectMixer/Source/ObjectMixer/Private/Views/List/SObjectMixerEditorListRow.cpp
@@ -976,8 +976,10 @@ void SObjectMixerEditorListRow::OnPropertyChanged(const FPropertyChangedEvent& E
 			 Event.ChangeType == EPropertyChangeType::Interactive ?
 			 	EPropertyValueSetFlags::InteractiveChange : EPropertyValueSetFlags::DefaultFlags;
 
-		const FObjectMixerEditorListRow::FPropertyPropagationInfo PropagationInfo(
-			{PinnedItem->GetUniqueIdentifier(), PropertyName, Flag});
+		FObjectMixerEditorListRow::FPropertyPropagationInfo PropagationInfo; 
+		PropagationInfo.RowIdentifier = PinnedItem->GetUniqueIdentifier();
+		PropagationInfo.PropertyName = PropertyName;
+		PropagationInfo.PropertyValueSetFlags = Flag;
 		
 		if (Flag == EPropertyValueSetFlags::InteractiveChange)
 		{
