Object Mixer:

-Fix: Unreal Editor crashes when the user clicks on the Create Folder Icon while having selected a child light

diff --git a/Engine/Plugins/Editor/ObjectMixer/ObjectMixer/Source/ObjectMixer/Private/Views/List/SObjectMixerEditorList.cpp b/Engine/Plugins/Editor/ObjectMixer/ObjectMixer/Source/ObjectMixer/Private/Views/List/SObjectMixerEditorList.cpp
--- a/Engine/Plugins/Editor/ObjectMixer/ObjectMixer/Source/ObjectMixer/Private/Views/List/SObjectMixerEditorList.cpp
+++ b/Engine/Plugins/Editor/ObjectMixer/ObjectMixer/Source/ObjectMixer/Private/Views/List/SObjectMixerEditorList.cpp
@@ -344,7 +344,8 @@ void SObjectMixerEditorList::OnRequestNewFolder(TOptional<FFolder> ExplicitParen
 				if (SelectedRow->GetRowType() == FObjectMixerEditorListRow::MatchingObject ||
 					SelectedRow->GetRowType() == FObjectMixerEditorListRow::ContainerObject)
 				{
-					if (FObjectMixerEditorListRowPtr PinnedParent = SelectedRow->GetDirectParentRow().Pin();
+					FObjectMixerEditorListRowPtr PinnedParent = SelectedRow->GetDirectParentRow().Pin();
+					if (PinnedParent.IsValid() &&
 						PinnedParent->GetRowType() == FObjectMixerEditorListRow::Folder &&
 						SelectedTreeItems.Contains(PinnedParent))
 					{
