Object Mixer:

-Add UMG Wrapper
-Fix: "Create Folder" button in Light Mixer (with nothing selected) creates folder in Outliner but not visible in Light Mixer

diff --git a/Engine/Plugins/Editor/ObjectMixer/LightMixer/Source/LightMixer/Private/LightMixerModule.cpp b/Engine/Plugins/Editor/ObjectMixer/LightMixer/Source/LightMixer/Private/LightMixerModule.cpp
--- a/Engine/Plugins/Editor/ObjectMixer/LightMixer/Source/LightMixer/Private/LightMixerModule.cpp
+++ b/Engine/Plugins/Editor/ObjectMixer/LightMixer/Source/LightMixer/Private/LightMixerModule.cpp
@@ -46,7 +46,7 @@ void FLightMixerModule::Initialize()
 	DefaultFilterClass = ULightMixerObjectFilter::StaticClass();	
 }
 
-FName FLightMixerModule::GetModuleName()
+FName FLightMixerModule::GetModuleName() const
 {
 	return "LightMixer";
 }
diff --git a/Engine/Plugins/Editor/ObjectMixer/LightMixer/Source/LightMixer/Public/LightMixerModule.h b/Engine/Plugins/Editor/ObjectMixer/LightMixer/Source/LightMixer/Public/LightMixerModule.h
--- a/Engine/Plugins/Editor/ObjectMixer/LightMixer/Source/LightMixer/Public/LightMixerModule.h
+++ b/Engine/Plugins/Editor/ObjectMixer/LightMixer/Source/LightMixer/Public/LightMixerModule.h
@@ -19,7 +19,7 @@ public:
 	
 	//~ Begin FObjectMixerEditorModule overrides
 	virtual void Initialize() override;
-	virtual FName GetModuleName() override;
+	virtual FName GetModuleName() const override;
 	virtual void SetupMenuItemVariables() override;
 	virtual FName GetTabSpawnerId() override;
 	virtual void RegisterSettings() const override;
diff --git a/Engine/Plugins/Editor/ObjectMixer/ObjectMixer/Source/ObjectMixer/ObjectMixerEditor.build.cs b/Engine/Plugins/Editor/ObjectMixer/ObjectMixer/Source/ObjectMixer/ObjectMixerEditor.build.cs
--- a/Engine/Plugins/Editor/ObjectMixer/ObjectMixer/Source/ObjectMixer/ObjectMixerEditor.build.cs
+++ b/Engine/Plugins/Editor/ObjectMixer/ObjectMixer/Source/ObjectMixer/ObjectMixerEditor.build.cs
@@ -47,6 +47,7 @@ public class ObjectMixerEditor : ModuleRules
 				"ToolWidgets",
 				"TypedElementFramework",
 				"TypedElementRuntime",
+				"UMG",
 				"UnrealEd",
 				"WorkspaceMenuStructure"
 			}
diff --git a/Engine/Plugins/Editor/ObjectMixer/ObjectMixer/Source/ObjectMixer/Private/ObjectMixerEditorModule.cpp b/Engine/Plugins/Editor/ObjectMixer/ObjectMixer/Source/ObjectMixer/Private/ObjectMixerEditorModule.cpp
--- a/Engine/Plugins/Editor/ObjectMixer/ObjectMixer/Source/ObjectMixer/Private/ObjectMixerEditorModule.cpp
+++ b/Engine/Plugins/Editor/ObjectMixer/ObjectMixer/Source/ObjectMixer/Private/ObjectMixerEditorModule.cpp
@@ -123,19 +123,28 @@ void FObjectMixerEditorModule::OpenProjectSettings()
 		.ShowViewer("Editor", "Plugins", "Object Mixer");
 }
 
-FName FObjectMixerEditorModule::GetModuleName()
+FName FObjectMixerEditorModule::GetModuleName() const
 {
 	return "ObjectMixerEditor";
 }
 
-TSharedPtr<SWidget> FObjectMixerEditorModule::MakeObjectMixerDialog() const
+TSharedPtr<SWidget> FObjectMixerEditorModule::MakeObjectMixerDialog(
+	TSubclassOf<UObjectMixerObjectFilter> InDefaultFilterClass)
 {
-	if (MainPanel.IsValid())
+	if (!MainPanel.IsValid())
+	{
+		MainPanel = MakeShared<FObjectMixerEditorMainPanel>(GetModuleName());
+		MainPanel->Init();
+	}
+	
+	const TSharedPtr<SWidget> ObjectMixerDialog = MainPanel->GetOrCreateWidget();
+
+	if (InDefaultFilterClass)
 	{
-		return MainPanel->GetOrCreateWidget();
+		MainPanel->SetDefaultFilterClass(InDefaultFilterClass);
 	}
 
-	return nullptr;
+	return ObjectMixerDialog;
 }
 
 void FObjectMixerEditorModule::RequestRebuildList() const
@@ -269,35 +278,23 @@ void FObjectMixerEditorModule::UnregisterSettings() const
 
 TSharedRef<SDockTab> FObjectMixerEditorModule::SpawnTab(const FSpawnTabArgs& Args)
 {
-	 return SpawnMainPanelTab();
-}
-
-TSharedPtr<FWorkspaceItem> FObjectMixerEditorModule::GetWorkspaceGroup()
-{
-	return WorkspaceGroup;
-}
-
-TSharedRef<SDockTab> FObjectMixerEditorModule::SpawnMainPanelTab()
-{
-	MainPanel = MakeShared<FObjectMixerEditorMainPanel>(GetModuleName());
-	MainPanel->Init();
-	
 	const TSharedRef<SDockTab> DockTab =
 		SNew(SDockTab)
 		.Label(TabLabel)
 		.TabRole(ETabRole::NomadTab)
 	;
-	const TSharedPtr<SWidget> ObjectMixerDialog = MakeObjectMixerDialog();
-	DockTab->SetContent(ObjectMixerDialog ? ObjectMixerDialog.ToSharedRef() : SNullWidget::NullWidget);
 
-	if (DefaultFilterClass)
-	{
-		MainPanel->SetDefaultFilterClass(DefaultFilterClass);
-	}
+	const TSharedPtr<SWidget> ObjectMixerDialog = MakeObjectMixerDialog(DefaultFilterClass);
+	DockTab->SetContent(ObjectMixerDialog ? ObjectMixerDialog.ToSharedRef() : SNullWidget::NullWidget);
 			
 	return DockTab;
 }
 
+TSharedPtr<FWorkspaceItem> FObjectMixerEditorModule::GetWorkspaceGroup()
+{
+	return WorkspaceGroup;
+}
+
 void FObjectMixerEditorModule::BindDelegates()
 {
 	check(GEngine && GEditor);
diff --git a/Engine/Plugins/Editor/ObjectMixer/ObjectMixer/Source/ObjectMixer/Private/Views/List/ObjectMixerEditorList.cpp b/Engine/Plugins/Editor/ObjectMixer/ObjectMixer/Source/ObjectMixer/Private/Views/List/ObjectMixerEditorList.cpp
--- a/Engine/Plugins/Editor/ObjectMixer/ObjectMixer/Source/ObjectMixer/Private/Views/List/ObjectMixerEditorList.cpp
+++ b/Engine/Plugins/Editor/ObjectMixer/ObjectMixer/Source/ObjectMixer/Private/Views/List/ObjectMixerEditorList.cpp
@@ -103,6 +103,26 @@ void FObjectMixerEditorList::RequestSyncEditorSelectionToListSelection() const
 	}
 }
 
+TArray<FObjectMixerEditorListRowPtr> FObjectMixerEditorList::GetSelectedTreeViewItems() const
+{
+	if (ListWidget.IsValid())
+	{
+		return ListWidget->GetSelectedTreeViewItems();
+	}
+
+	return {};
+}
+
+int32 FObjectMixerEditorList::GetSelectedTreeViewItemCount() const
+{
+	if (ListWidget.IsValid())
+	{
+		return ListWidget->GetSelectedTreeViewItemCount();
+	}
+
+	return INDEX_NONE;
+}
+
 void FObjectMixerEditorList::OnRequestNewFolder(TOptional<FFolder> ExplicitParentFolder)
 {
 	if (ListWidget.IsValid())
diff --git a/Engine/Plugins/Editor/ObjectMixer/ObjectMixer/Source/ObjectMixer/Private/Views/MainPanel/SObjectMixerEditorMainPanel.cpp b/Engine/Plugins/Editor/ObjectMixer/ObjectMixer/Source/ObjectMixer/Private/Views/MainPanel/SObjectMixerEditorMainPanel.cpp
--- a/Engine/Plugins/Editor/ObjectMixer/ObjectMixer/Source/ObjectMixer/Private/Views/MainPanel/SObjectMixerEditorMainPanel.cpp
+++ b/Engine/Plugins/Editor/ObjectMixer/ObjectMixer/Source/ObjectMixer/Private/Views/MainPanel/SObjectMixerEditorMainPanel.cpp
@@ -172,18 +172,26 @@ TSharedRef<SWidget> SObjectMixerEditorMainPanel::GenerateToolbar()
 	[
 		SNew(SButton)
 		.ButtonStyle(FAppStyle::Get(), "SimpleButton")
-		.ToolTipText(LOCTEXT("CreateFolderToolTip", "Create a new folder containing the current selection"))
+		.ToolTipText_Lambda([this] ()
+		{
+			return CanCreateFolder() ? LOCTEXT("CreateFolderToolTip", "Create a new folder containing the current selection") :
+				LOCTEXT("CannotCreateFolder","Please select an item in the tree view in order to create a folder.");
+		})
 		.OnClicked_Lambda([this] ()
 		{
-			if (TSharedPtr<FObjectMixerEditorList> PinnedList = GetMainPanelModel().Pin()->GetEditorListModel().Pin())
+			if (CanCreateFolder())
 			{
-				PinnedList->OnRequestNewFolder();
+				if (const TSharedPtr<FObjectMixerEditorList> PinnedList = GetMainPanelModel().Pin()->GetEditorListModel().Pin())
+				{
+					PinnedList->OnRequestNewFolder();
 
-				return FReply::Handled();
+					return FReply::Handled();
+				}
 			}
 
 			return FReply::Unhandled();
 		})
+		.IsEnabled(this, &SObjectMixerEditorMainPanel::CanCreateFolder)
 		[
 			SNew(SImage)
 			.ColorAndOpacity(FSlateColor::UseForeground())
@@ -215,6 +223,16 @@ TSharedRef<SWidget> SObjectMixerEditorMainPanel::GenerateToolbar()
 	return ToolbarBox;
 }
 
+bool SObjectMixerEditorMainPanel::CanCreateFolder() const
+{
+	if (TSharedPtr<FObjectMixerEditorList> PinnedList = GetMainPanelModel().Pin()->GetEditorListModel().Pin())
+	{
+		return PinnedList->GetSelectedTreeViewItemCount() > 0;
+	}
+
+	return false;
+}
+
 TSharedRef<SWidget> SObjectMixerEditorMainPanel::OnGenerateAddObjectButtonMenu() const
 {
 	TSet<UClass*> ClassesToPlace;
@@ -367,7 +385,10 @@ TSharedRef<SWidget> SObjectMixerEditorMainPanel::OnGenerateFilterClassMenu()
 					continue;
 				}
 
-				const bool bIsDefaultClass = MainPanelModel.Pin()->GetDefaultFilterClass() == AssetClassMap.Class;
+				const bool bIsDefaultClass =
+					// If this is a made-to-purpose sub-plugin of Object Mixer, don't allow default class to be disabled
+					GetMainPanelModel().Pin()->GetModuleName() != FObjectMixerEditorModule::BaseObjectMixerModuleName &&
+					MainPanelModel.Pin()->GetDefaultFilterClass() == AssetClassMap.Class;
 				
 				const FText TooltipText = bIsDefaultClass ?
 					FText::Format(
diff --git a/Engine/Plugins/Editor/ObjectMixer/ObjectMixer/Source/ObjectMixer/Private/Views/Widgets/ObjectMixerEditorUWidget.cpp b/Engine/Plugins/Editor/ObjectMixer/ObjectMixer/Source/ObjectMixer/Private/Views/Widgets/ObjectMixerEditorUWidget.cpp
new file mode 100644
--- /dev/null
+++ b/Engine/Plugins/Editor/ObjectMixer/ObjectMixer/Source/ObjectMixer/Private/Views/Widgets/ObjectMixerEditorUWidget.cpp
@@ -0,0 +1,47 @@
+// Copyright Epic Games, Inc. All Rights Reserved.
+
+#include "Views/Widgets/ObjectMixerEditorUWidget.h"
+
+#include "ObjectMixerEditorModule.h"
+
+#include "Editor.h"
+#include "Editor/EditorEngine.h"
+#include "Widgets/Text/STextBlock.h"
+
+#define LOCTEXT_NAMESPACE "ObjectMixerEditor"
+
+const FText UObjectMixerEditorUWidget::GetPaletteCategory()
+{
+	return LOCTEXT("Editor", "Editor");
+}
+
+TSharedRef<SWidget> UObjectMixerEditorUWidget::RebuildWidget()
+{
+	TSharedRef<SBorder> Border = SNew(SBorder)
+		.Padding(0.0f)
+		.HAlign(HAlign_Fill)
+		.VAlign(VAlign_Fill)
+		.BorderImage(FAppStyle::GetBrush("NoBorder"))
+	;
+	
+	if (GEditor)
+	{
+		if (!GEditor->IsPlayingSessionInEditor())
+		{
+			const TSharedPtr<SWidget> Widget =
+				FObjectMixerEditorModule::Get().MakeObjectMixerDialog(ObjectMixerWidgetUserConfig.DefaultFilterClass);
+			Border->SetContent(Widget.IsValid() ? Widget.ToSharedRef() : SNullWidget::NullWidget);
+
+			return Border;
+		}
+	}
+
+	Border->SetContent(
+		SNew(STextBlock)
+		.Text(LOCTEXT("EditorWidget", "Editor Only"))
+	);
+
+	return Border;
+}
+
+#undef LOCTEXT_NAMESPACE
diff --git a/Engine/Plugins/Editor/ObjectMixer/ObjectMixer/Source/ObjectMixer/Public/ObjectMixerEditorModule.h b/Engine/Plugins/Editor/ObjectMixer/ObjectMixer/Source/ObjectMixer/Public/ObjectMixerEditorModule.h
--- a/Engine/Plugins/Editor/ObjectMixer/ObjectMixer/Source/ObjectMixer/Public/ObjectMixerEditorModule.h
+++ b/Engine/Plugins/Editor/ObjectMixer/ObjectMixer/Source/ObjectMixer/Public/ObjectMixerEditorModule.h
@@ -26,9 +26,10 @@ public:
 
 	static void OpenProjectSettings();
 
-	virtual FName GetModuleName();
+	virtual FName GetModuleName() const;
 	
-	virtual TSharedPtr<SWidget> MakeObjectMixerDialog() const;
+	virtual TSharedPtr<SWidget> MakeObjectMixerDialog(
+		TSubclassOf<UObjectMixerObjectFilter> InDefaultFilterClass = nullptr);
 
 	/**
 	 * Regenerate the list items and refresh the list. Call when adding or removing variables.
@@ -76,8 +77,6 @@ public:
 
 protected:
 
-	virtual TSharedRef<SDockTab> SpawnMainPanelTab();
-
 	virtual void BindDelegates();
 	
 	/** Lives for as long as the module is loaded. */
diff --git a/Engine/Plugins/Editor/ObjectMixer/ObjectMixer/Source/ObjectMixer/Public/Views/List/ObjectMixerEditorList.h b/Engine/Plugins/Editor/ObjectMixer/ObjectMixer/Source/ObjectMixer/Public/Views/List/ObjectMixerEditorList.h
--- a/Engine/Plugins/Editor/ObjectMixer/ObjectMixer/Source/ObjectMixer/Public/Views/List/ObjectMixerEditorList.h
+++ b/Engine/Plugins/Editor/ObjectMixer/ObjectMixer/Source/ObjectMixer/Public/Views/List/ObjectMixerEditorList.h
@@ -2,6 +2,7 @@
 
 #pragma once
 
+#include "ObjectMixerEditorListRow.h"
 #include "Views/MainPanel/ObjectMixerEditorMainPanel.h"
 
 #include "Templates/SharedPointer.h"
@@ -44,6 +45,9 @@ public:
 
 	void RequestSyncEditorSelectionToListSelection() const;
 
+	[[nodiscard]] TArray<FObjectMixerEditorListRowPtr> GetSelectedTreeViewItems() const;
+	int32 GetSelectedTreeViewItemCount() const;
+
 	void OnRequestNewFolder(TOptional<FFolder> ExplicitParentFolder = TOptional<FFolder>());
 
 	void OnRequestMoveFolder(const FFolder& FolderToMove, const FFolder& TargetNewParentFolder);
diff --git a/Engine/Plugins/Editor/ObjectMixer/ObjectMixer/Source/ObjectMixer/Public/Views/MainPanel/SObjectMixerEditorMainPanel.h b/Engine/Plugins/Editor/ObjectMixer/ObjectMixer/Source/ObjectMixer/Public/Views/MainPanel/SObjectMixerEditorMainPanel.h
--- a/Engine/Plugins/Editor/ObjectMixer/ObjectMixer/Source/ObjectMixer/Public/Views/MainPanel/SObjectMixerEditorMainPanel.h
+++ b/Engine/Plugins/Editor/ObjectMixer/ObjectMixer/Source/ObjectMixer/Public/Views/MainPanel/SObjectMixerEditorMainPanel.h
@@ -29,7 +29,7 @@ public:
 
 	virtual FReply OnKeyDown( const FGeometry& MyGeometry, const FKeyEvent& InKeyEvent ) override;
 
-	TWeakPtr<FObjectMixerEditorMainPanel> GetMainPanelModel()
+	TWeakPtr<FObjectMixerEditorMainPanel> GetMainPanelModel() const
 	{
 		return MainPanelModel;
 	}
@@ -84,6 +84,7 @@ private:
 	TSharedPtr<class SComboButton> ViewOptionsComboButton;
 
 	TSharedRef<SWidget> GenerateToolbar();
+	bool CanCreateFolder() const;
 	TSharedRef<SWidget> OnGenerateAddObjectButtonMenu() const;
 
 	TArray<FFilterClassSelectionInfo> FilterClassSelectionInfos;
diff --git a/Engine/Plugins/Editor/ObjectMixer/ObjectMixer/Source/ObjectMixer/Public/Views/Widgets/ObjectMixerEditorUWidget.h b/Engine/Plugins/Editor/ObjectMixer/ObjectMixer/Source/ObjectMixer/Public/Views/Widgets/ObjectMixerEditorUWidget.h
new file mode 100644
--- /dev/null
+++ b/Engine/Plugins/Editor/ObjectMixer/ObjectMixer/Source/ObjectMixer/Public/Views/Widgets/ObjectMixerEditorUWidget.h
@@ -0,0 +1,45 @@
+// Copyright Epic Games, Inc. All Rights Reserved.
+
+#pragma once
+
+#include "ObjectFilter/ObjectMixerEditorObjectFilter.h"
+
+#include "Components/Widget.h"
+
+#include "ObjectMixerEditorUWidget.generated.h"
+
+// User Configurable Variables
+USTRUCT(BlueprintType)
+struct FObjectMixerWidgetUserConfig
+{
+	GENERATED_BODY()
+
+	/**
+	 * Sets the default filter class to determine what objects and properties to display in this Object Mixer instance.
+	 */
+	UPROPERTY(EditAnywhere, BlueprintReadWrite, Category = "Object Mixer")
+	TSubclassOf<UObjectMixerObjectFilter> DefaultFilterClass;
+};
+
+/**
+ * A UMG widget wrapper for the Object Mixer widget. Uses the same data as the Generic Object Mixer Instance.
+ * Only useful in the editor. It is not compatible at runtime.
+ */
+UCLASS(meta = (DisplayName="Object Mixer (Generic)"))
+class OBJECTMIXEREDITOR_API UObjectMixerEditorUWidget : public UWidget
+{
+	GENERATED_BODY()
+
+public:
+
+#if WITH_EDITOR
+	virtual const FText GetPaletteCategory() override;
+#endif
+
+	// UWidget
+	virtual TSharedRef<SWidget> RebuildWidget() override;
+	// End of UWidget
+
+	UPROPERTY(EditAnywhere, BlueprintReadWrite, Category = "Object Mixer")
+	FObjectMixerWidgetUserConfig ObjectMixerWidgetUserConfig;
+};
