Object Mixer:

-Fix: Shift Multi-Selection Bug

diff --git a/Engine/Plugins/Editor/ObjectMixer/ObjectMixer/Source/ObjectMixer/Private/ObjectFilter/ObjectMixerEditorObjectFilter.cpp b/Engine/Plugins/Editor/ObjectMixer/ObjectMixer/Source/ObjectMixer/Private/ObjectFilter/ObjectMixerEditorObjectFilter.cpp
--- a/Engine/Plugins/Editor/ObjectMixer/ObjectMixer/Source/ObjectMixer/Private/ObjectFilter/ObjectMixerEditorObjectFilter.cpp
+++ b/Engine/Plugins/Editor/ObjectMixer/ObjectMixer/Source/ObjectMixer/Private/ObjectFilter/ObjectMixerEditorObjectFilter.cpp
@@ -135,6 +135,11 @@ bool UObjectMixerObjectFilter::ShouldIncludeUnsupportedProperties() const
 	return false;
 }
 
+TSet<FName> UObjectMixerObjectFilter::GetPropertiesThatRequireListRefresh() const
+{
+	return { "Mobility" };
+}
+
 TSet<UClass*> UObjectMixerObjectFilter::GetParentAndChildClassesFromSpecifiedClasses(
 	const TSet<UClass*>& InSpecifiedClasses, EObjectMixerInheritanceInclusionOptions Options)
 {
diff --git a/Engine/Plugins/Editor/ObjectMixer/ObjectMixer/Source/ObjectMixer/Private/ObjectMixerEditorModule.cpp b/Engine/Plugins/Editor/ObjectMixer/ObjectMixer/Source/ObjectMixer/Private/ObjectMixerEditorModule.cpp
--- a/Engine/Plugins/Editor/ObjectMixer/ObjectMixer/Source/ObjectMixer/Private/ObjectMixerEditorModule.cpp
+++ b/Engine/Plugins/Editor/ObjectMixer/ObjectMixer/Source/ObjectMixer/Private/ObjectMixerEditorModule.cpp
@@ -52,12 +52,7 @@ UWorld* FObjectMixerEditorModule::GetWorld()
 }
 
 void FObjectMixerEditorModule::Initialize()
-{
-	PropertiesThatRequireRefresh =
-		{
-			"LightingChannels"
-		};
-	
+{	
 	BindDelegates();
 	
 	SetupMenuItemVariables();
@@ -134,7 +129,7 @@ TSharedPtr<SWidget> FObjectMixerEditorModule::MakeObjectMixerDialog(
 	if (!MainPanel.IsValid())
 	{
 		MainPanel = MakeShared<FObjectMixerEditorMainPanel>(GetModuleName());
-		MainPanel->Init();
+		MainPanel->Initialize();
 	}
 	
 	const TSharedPtr<SWidget> ObjectMixerDialog = MainPanel->GetOrCreateWidget();
@@ -339,12 +334,12 @@ void FObjectMixerEditorModule::BindDelegates()
 	{
 		RequestSyncEditorSelectionToListSelection();
 	}));
-
+	
 	DelegateHandles.Add(FEditorDelegates::MapChange.AddLambda([this](uint32)
 	{
 		RequestRebuildList();
 	}));
-
+	
 	DelegateHandles.Add(FEditorDelegates::PostUndoRedo.AddLambda([this]()
 	{
 		RequestRebuildList();
@@ -352,7 +347,7 @@ void FObjectMixerEditorModule::BindDelegates()
 		// Because we can undo/redo collection adds and removes, we need to save the data after each
 		GetMutableDefault<UObjectMixerEditorSerializedData>()->SaveConfig();
 	}));
-
+	
 	DelegateHandles.Add(FCoreUObjectDelegates::OnObjectTransacted.AddLambda([this](UObject*, const FTransactionObjectEvent& Event)
 	{		
 		if (Event.GetEventType() == ETransactionObjectEventType::Finalized)
@@ -361,21 +356,18 @@ void FObjectMixerEditorModule::BindDelegates()
 			{
 				RequestRebuildList();
 			}
-			else
+			else if (const TSet<FName> PropertiesThatRequireRefresh = GetPropertiesThatRequireRefresh(); PropertiesThatRequireRefresh.Num() > 0)
 			{
 				const TSet<FName> ChangedPropertyNames = TSet<FName>(Event.GetChangedProperties());
-				const bool bPropertyChangedThatRequiresRefresh =
-					ChangedPropertyNames.Difference(PropertiesThatRequireRefresh).Num() > 0 ||
-						PropertiesThatRequireRefresh.Difference(ChangedPropertyNames).Num() > 0;
-
-				if (bPropertyChangedThatRequiresRefresh)
+	
+				if (ChangedPropertyNames.Intersect(PropertiesThatRequireRefresh).Num() > 0)
 				{
 					RequestRebuildList();
 				}
 			}
 		}
 	}));
-
+	
 	FLevelEditorModule& LevelEditorModule = FModuleManager::LoadModuleChecked<FLevelEditorModule>("LevelEditor");
 	LevelEditorModule.OnComponentsEdited().AddLambda([this]()
 	{
@@ -383,6 +375,21 @@ void FObjectMixerEditorModule::BindDelegates()
 	});
 }
 
+TSet<FName> FObjectMixerEditorModule::GetPropertiesThatRequireRefresh() const
+{
+	TSet<FName> ReturnValue;
+
+	if (MainPanel.IsValid())
+	{
+		for (const TObjectPtr<UObjectMixerObjectFilter> Instance : MainPanel->GetObjectFilterInstances())
+		{
+			ReturnValue.Append(Instance->GetPropertiesThatRequireListRefresh());
+		}
+	}
+
+	return ReturnValue;
+}
+
 const TSubclassOf<UObjectMixerObjectFilter>& FObjectMixerEditorModule::GetDefaultFilterClass() const
 {
 	return DefaultFilterClass;
diff --git a/Engine/Plugins/Editor/ObjectMixer/ObjectMixer/Source/ObjectMixer/Private/Views/MainPanel/ObjectMixerEditorMainPanel.cpp b/Engine/Plugins/Editor/ObjectMixer/ObjectMixer/Source/ObjectMixer/Private/Views/MainPanel/ObjectMixerEditorMainPanel.cpp
--- a/Engine/Plugins/Editor/ObjectMixer/ObjectMixer/Source/ObjectMixer/Private/Views/MainPanel/ObjectMixerEditorMainPanel.cpp
+++ b/Engine/Plugins/Editor/ObjectMixer/ObjectMixer/Source/ObjectMixer/Private/Views/MainPanel/ObjectMixerEditorMainPanel.cpp
@@ -9,7 +9,7 @@
 #include "Views/List/ObjectMixerEditorList.h"
 #include "Views/MainPanel/SObjectMixerEditorMainPanel.h"
 
-void FObjectMixerEditorMainPanel::Init()
+void FObjectMixerEditorMainPanel::Initialize()
 {
 	RegenerateListModel();
 	RegisterAndMapContextMenuCommands();
@@ -42,32 +42,6 @@ void FObjectMixerEditorMainPanel::RegisterAndMapContextMenuCommands()
 	ObjectMixerElementEditCommands->MapAction(FGenericCommands::Get().Rename,
 		FUIAction(FExecuteAction::CreateRaw(this, &FObjectMixerEditorMainPanel::OnRenameCommand))
 	);
-
-	ObjectMixerFolderEditCommands = MakeShared<FUICommandList>();
-
-	ObjectMixerElementEditCommands->MapAction(FGenericCommands::Get().Cut,
-		FExecuteAction::CreateStatic( &FLevelEditorActionCallbacks::ExecuteExecCommand, FString( TEXT("EDIT CUT") ) ),
-		FCanExecuteAction::CreateStatic( &FLevelEditorActionCallbacks::Cut_CanExecute )
-	);
-	ObjectMixerElementEditCommands->MapAction(FGenericCommands::Get().Copy,
-		FExecuteAction::CreateStatic( &FLevelEditorActionCallbacks::ExecuteExecCommand, FString( TEXT("EDIT COPY") ) ),
-		FCanExecuteAction::CreateStatic( &FLevelEditorActionCallbacks::Copy_CanExecute )
-	);
-	ObjectMixerElementEditCommands->MapAction(FGenericCommands::Get().Paste,
-		FExecuteAction::CreateStatic( &FLevelEditorActionCallbacks::ExecuteExecCommand, FString( TEXT("EDIT PASTE") ) ),
-		FCanExecuteAction::CreateStatic( &FLevelEditorActionCallbacks::Paste_CanExecute )
-	);
-	ObjectMixerElementEditCommands->MapAction(FGenericCommands::Get().Duplicate,
-		FExecuteAction::CreateStatic( &FLevelEditorActionCallbacks::ExecuteExecCommand, FString( TEXT("DUPLICATE") ) ),
-		FCanExecuteAction::CreateStatic( &FLevelEditorActionCallbacks::Duplicate_CanExecute )
-	);
-	ObjectMixerElementEditCommands->MapAction(FGenericCommands::Get().Delete,
-		FExecuteAction::CreateStatic( &FLevelEditorActionCallbacks::ExecuteExecCommand, FString( TEXT("DELETE") ) ),
-		FCanExecuteAction::CreateStatic( &FLevelEditorActionCallbacks::Delete_CanExecute )
-	);
-	ObjectMixerElementEditCommands->MapAction(FGenericCommands::Get().Rename,
-		FUIAction(FExecuteAction::CreateRaw(this, &FObjectMixerEditorMainPanel::OnRenameCommand))
-	);
 }
 
 TSharedRef<SWidget> FObjectMixerEditorMainPanel::GetOrCreateWidget()
diff --git a/Engine/Plugins/Editor/ObjectMixer/ObjectMixer/Source/ObjectMixer/Public/ObjectFilter/ObjectMixerEditorObjectFilter.h b/Engine/Plugins/Editor/ObjectMixer/ObjectMixer/Source/ObjectMixer/Public/ObjectFilter/ObjectMixerEditorObjectFilter.h
--- a/Engine/Plugins/Editor/ObjectMixer/ObjectMixer/Source/ObjectMixer/Public/ObjectFilter/ObjectMixerEditorObjectFilter.h
+++ b/Engine/Plugins/Editor/ObjectMixer/ObjectMixer/Source/ObjectMixer/Public/ObjectFilter/ObjectMixerEditorObjectFilter.h
@@ -140,6 +140,12 @@ public:
 	 */
 	virtual bool ShouldIncludeUnsupportedProperties() const;
 
+	/**
+	 * If a property is changed that has a name found in this set, the panel will be refreshed.
+	 * Add a property name to this list if you expect the list to change in some way after changing that property.
+	 */
+	virtual TSet<FName> GetPropertiesThatRequireListRefresh() const;
+
 	static TSet<UClass*> GetParentAndChildClassesFromSpecifiedClasses(
 		const TSet<UClass*>& InSpecifiedClasses, EObjectMixerInheritanceInclusionOptions Options);
 
@@ -325,5 +331,17 @@ public:
 	{
 		return Super::ShouldIncludeUnsupportedProperties();
 	}
+	
+	/**
+	 * If a property is changed that has a name found in this set, the panel will be refreshed.
+	 * Add a property name to this list if you expect the list to change in some way after changing that property.
+	 */
+	UFUNCTION(BlueprintNativeEvent, BlueprintCallable, Category = "Object Mixer")
+	TSet<FName> GetPropertiesThatRequireListRefresh() const override;
+
+	TSet<FName> GetPropertiesThatRequireListRefresh_Implementation() const
+	{
+		return Super::GetPropertiesThatRequireListRefresh();
+	}
 
 };
diff --git a/Engine/Plugins/Editor/ObjectMixer/ObjectMixer/Source/ObjectMixer/Public/ObjectMixerEditorModule.h b/Engine/Plugins/Editor/ObjectMixer/ObjectMixer/Source/ObjectMixer/Public/ObjectMixerEditorModule.h
--- a/Engine/Plugins/Editor/ObjectMixer/ObjectMixer/Source/ObjectMixer/Public/ObjectMixerEditorModule.h
+++ b/Engine/Plugins/Editor/ObjectMixer/ObjectMixer/Source/ObjectMixer/Public/ObjectMixerEditorModule.h
@@ -78,6 +78,9 @@ public:
 protected:
 
 	virtual void BindDelegates();
+
+	/** If a property is changed that has a name found in this set, the panel will be refreshed. */
+	TSet<FName> GetPropertiesThatRequireRefresh() const;
 	
 	/** Lives for as long as the module is loaded. */
 	TSharedPtr<FObjectMixerEditorMainPanel> MainPanel;
@@ -99,9 +102,6 @@ protected:
 	
 	TSet<FDelegateHandle> DelegateHandles;
 
-	/** If a property is changed that has a name found in this set, the panel will be refreshed. */
-	TSet<FName> PropertiesThatRequireRefresh;
-
 private:
 
 	TSharedPtr<FWorkspaceItem> WorkspaceGroup;
diff --git a/Engine/Plugins/Editor/ObjectMixer/ObjectMixer/Source/ObjectMixer/Public/Views/MainPanel/ObjectMixerEditorMainPanel.h b/Engine/Plugins/Editor/ObjectMixer/ObjectMixer/Source/ObjectMixer/Public/Views/MainPanel/ObjectMixerEditorMainPanel.h
--- a/Engine/Plugins/Editor/ObjectMixer/ObjectMixer/Source/ObjectMixer/Public/Views/MainPanel/ObjectMixerEditorMainPanel.h
+++ b/Engine/Plugins/Editor/ObjectMixer/ObjectMixer/Source/ObjectMixer/Public/Views/MainPanel/ObjectMixerEditorMainPanel.h
@@ -29,9 +29,9 @@ public:
 	: ModuleName(InModuleName)
 	{}
 
-	~FObjectMixerEditorMainPanel() = default;
+	virtual ~FObjectMixerEditorMainPanel() override = default;
 
-	void Init();
+	void Initialize();
 	
 	void RegisterAndMapContextMenuCommands();
 
@@ -45,7 +45,7 @@ public:
 	void RequestRebuildList() const;
 
 	/**
-	 * Refresh filters and sorting.
+	 * Refresh list filters and sorting.
 	 * Useful for when the list state has gone stale but the variable count has not changed.
 	 */
 	void RefreshList() const;
