Object Mixer:

-Disable Ability To Rename Lights contained in Blueprints or other Actors

diff --git a/Engine/Plugins/Editor/ObjectMixer/ObjectMixer/Source/ObjectMixer/Private/Views/List/SObjectMixerEditorListRow.cpp b/Engine/Plugins/Editor/ObjectMixer/ObjectMixer/Source/ObjectMixer/Private/Views/List/SObjectMixerEditorListRow.cpp
--- a/Engine/Plugins/Editor/ObjectMixer/ObjectMixer/Source/ObjectMixer/Private/Views/List/SObjectMixerEditorListRow.cpp
+++ b/Engine/Plugins/Editor/ObjectMixer/ObjectMixer/Source/ObjectMixer/Private/Views/List/SObjectMixerEditorListRow.cpp
@@ -35,7 +35,7 @@ class SInlineEditableRowNameCellWidget : public SCompoundWidget
 {
 public:
 	
-	SLATE_BEGIN_ARGS(SObjectMixerEditorListRow)
+	SLATE_BEGIN_ARGS(SInlineEditableRowNameCellWidget)
 	{}
 
 	SLATE_END_ARGS()
@@ -45,7 +45,7 @@ public:
 		Item = InRow;
 		HybridChild = InHybridChild;
 
-		InRow->OnRenameCommand().BindRaw(this, &SInlineEditableRowNameCellWidget::EnterEditingMode);
+		InRow->OnRenameCommand().BindRaw(this, &SInlineEditableRowNameCellWidget::OnRenameRequested);
 		
 		TSharedRef<SHorizontalBox> HBox = SNew(SHorizontalBox);
 		
@@ -110,6 +110,7 @@ public:
 				.HighlightText(this, &SInlineEditableRowNameCellWidget::GetHighlightText)
 				.IsSelected_Raw(this, &SInlineEditableRowNameCellWidget::GetIsSelectedExclusively)
 				.OnTextCommitted(this, &SInlineEditableRowNameCellWidget::OnTextCommitted)
+				.IsReadOnly(this, &SInlineEditableRowNameCellWidget::IsReadOnly)
 			];
 		}
 		
@@ -151,6 +152,20 @@ public:
 		}
 	}
 
+	bool IsReadOnly() const
+	{
+		const UActorComponent* Component = Cast<UActorComponent>(Item.Pin()->GetObject());
+		return Component != nullptr;
+	}
+
+	void OnRenameRequested() const
+	{
+		if (!IsReadOnly())
+		{
+			EnterEditingMode();
+		}
+	}
+
 private:
 
 	bool IsValidComponentRename(const UActorComponent* ComponentInstance, const FText& InNewText) const
