Object Mixer:

-Fix: Changing Parameters on a Multi-Selection Only Change Values Of Lights On Screen
-Fix: Blueprint Lights Parameters Switch To "Multiple Values" When Transformed
-Fix: Multi-Select Blueprint/Actor w/Lights and Light Parameter Change Bug

diff --git a/Engine/Plugins/Editor/ObjectMixer/ObjectMixer/Source/ObjectMixer/Private/Views/List/ObjectMixerEditorListRow.cpp b/Engine/Plugins/Editor/ObjectMixer/ObjectMixer/Source/ObjectMixer/Private/Views/List/ObjectMixerEditorListRow.cpp
--- a/Engine/Plugins/Editor/ObjectMixer/ObjectMixer/Source/ObjectMixer/Private/Views/List/ObjectMixerEditorListRow.cpp
+++ b/Engine/Plugins/Editor/ObjectMixer/ObjectMixer/Source/ObjectMixer/Private/Views/List/ObjectMixerEditorListRow.cpp
@@ -4,6 +4,7 @@
 
 #include "Algo/AllOf.h"
 #include "ClassIconFinder.h"
+#include "ObjectMixerEditorLog.h"
 #include "ObjectMixerEditorSerializedData.h"
 #include "ScopedTransaction.h"
 #include "GameFramework/Actor.h"
@@ -573,65 +574,51 @@ void SetValueOnSelectedItems(
 	{
 		FScopedTransaction Transaction(
 			NSLOCTEXT("ObjectMixerEditor","OnPropertyChangedTransaction", "Object Mixer - Bulk Edit Selected Row Properties") );
-				
+		
 		for (const TSharedPtr<FObjectMixerEditorListRow>& SelectedRow : OtherSelectedItems)
 		{
 			const FObjectMixerEditorListRowPtr SelectedHybridRow = SelectedRow->GetHybridChild();
 			const FObjectMixerEditorListRowPtr RowToUse = SelectedHybridRow.IsValid() ? SelectedHybridRow : SelectedRow;
 
-			// Use handles if valid, otherwise use ImportText
 			if (RowToUse == PinnedItem)
 			{
-				return;
+				continue;
+			}
+
+			// Skip folders
+			if (RowToUse->GetRowType() == FObjectMixerEditorListRow::Folder)
+			{
+				continue;
 			}
 
 			UObject* ObjectToModify = RowToUse->GetObject();
 			
+			if (IsValid(ObjectToModify))
+			{
+				ObjectToModify->Modify();
+			}
+			else
+			{
+				UE_LOG(LogObjectMixerEditor, Warning, TEXT("%hs: Row '%s' has no valid associated object to modify."), __FUNCTION__, *RowToUse->GetDisplayName().ToString());
+				return;
+			}
+		
+			// Use handles if valid, otherwise use ImportText
 			if (const TWeakPtr<IPropertyHandle>* SelectedHandlePtr = RowToUse->PropertyNamesToHandles.Find(PropertyName))
 			{
 				if (SelectedHandlePtr->IsValid())
-				{
-					if (ObjectToModify)
-					{
-						ObjectToModify->Modify();
-					}
-					
+				{				
 					SelectedHandlePtr->Pin()->SetValueFromFormattedString(ValueAsString, Flags);
-					continue;
 				}
 			}
-			
-			if (ObjectToModify)
+			else
 			{
-				if (FProperty* PropertyToChange = FindFProperty<FProperty>(ObjectToModify->GetClass(), PropertyName))
+				if (const FProperty* PropertyToChange = FindFProperty<FProperty>(ObjectToModify->GetClass(), PropertyName))
 				{
 					if (void* ValuePtr = PropertyToChange->ContainerPtrToValuePtr<void>(ObjectToModify))
 					{
-						EPropertyChangeType::Type ChangeType =
-							Flags == EPropertyValueSetFlags::InteractiveChange
-								? EPropertyChangeType::Interactive
-								: EPropertyChangeType::ValueSet;
-
 						// Set the actual property value
 						PropertyToChange->ImportText_Direct(*ValueAsString, ValuePtr, ObjectToModify, PPF_None);
-						FPropertyChangedEvent ChangeEvent(
-							PropertyToChange,
-							ChangeType,
-							MakeArrayView({ObjectToModify}));
-						ObjectToModify->PostEditChangeProperty(ChangeEvent);
-
-						// Propagate to outers
-						UObject* Outer = ObjectToModify->GetOuter();
-						while (Outer) 
-						{
-							FPropertyChangedEvent ActorChangeEvent(
-								PropertyToChange,
-								ChangeType,
-								MakeArrayView({Outer}));
-							Outer->PostEditChangeProperty(ActorChangeEvent);
-
-							Outer = Outer->GetOuter();
-						}
 					}
 				}
 			}
diff --git a/Engine/Plugins/Editor/ObjectMixer/ObjectMixer/Source/ObjectMixer/Private/Views/List/SObjectMixerEditorList.cpp b/Engine/Plugins/Editor/ObjectMixer/ObjectMixer/Source/ObjectMixer/Private/Views/List/SObjectMixerEditorList.cpp
--- a/Engine/Plugins/Editor/ObjectMixer/ObjectMixer/Source/ObjectMixer/Private/Views/List/SObjectMixerEditorList.cpp
+++ b/Engine/Plugins/Editor/ObjectMixer/ObjectMixer/Source/ObjectMixer/Private/Views/List/SObjectMixerEditorList.cpp
@@ -1333,7 +1333,7 @@ void SObjectMixerEditorList::PropagatePropertyChangesToSelectedRows()
 			TSharedPtr<STreeView<FObjectMixerEditorListRowPtr>>,
 			const TArray<FObjectMixerEditorListRow::FPropertyPropagationInfo>&);
 	
-	static LambdaType RecursivelyRestoreTreeState = [](
+	static LambdaType RecursivelyPropagatePropertyChangesToSelectedRows = [](
 		const TArray<FObjectMixerEditorListRowPtr>& InObjects,
 		TSharedPtr<STreeView<FObjectMixerEditorListRowPtr>> InTreeViewPtr,
 		const TArray<FObjectMixerEditorListRow::FPropertyPropagationInfo>& InPendingPropertyPropagations)
@@ -1353,7 +1353,7 @@ void SObjectMixerEditorList::PropagatePropertyChangesToSelectedRows()
 					TreeViewItem->PropagateChangesToSimilarSelectedRowProperties(*PendingPropagation);
 				}
 
-			RecursivelyRestoreTreeState(
+			RecursivelyPropagatePropertyChangesToSelectedRows(
 				TreeViewItem->GetChildRows(), InTreeViewPtr, InPendingPropertyPropagations
 			);
 		}
@@ -1361,7 +1361,7 @@ void SObjectMixerEditorList::PropagatePropertyChangesToSelectedRows()
 
 	if (GetSelectedTreeViewItemCount() > 1)
 	{
-		RecursivelyRestoreTreeState(
+			RecursivelyPropagatePropertyChangesToSelectedRows(
 		   TreeViewRootObjects, TreeViewPtr, PendingPropertyPropagations
 	   );
 	}
