Object Mixer:

-Fix: Renaming a collection will result in an identical collection tab

diff --git a/Engine/Plugins/Editor/ObjectMixer/ObjectMixer/Source/ObjectMixer/Private/Views/MainPanel/ObjectMixerEditorMainPanel.cpp b/Engine/Plugins/Editor/ObjectMixer/ObjectMixer/Source/ObjectMixer/Private/Views/MainPanel/ObjectMixerEditorMainPanel.cpp
--- a/Engine/Plugins/Editor/ObjectMixer/ObjectMixer/Source/ObjectMixer/Private/Views/MainPanel/ObjectMixerEditorMainPanel.cpp
+++ b/Engine/Plugins/Editor/ObjectMixer/ObjectMixer/Source/ObjectMixer/Private/Views/MainPanel/ObjectMixerEditorMainPanel.cpp
@@ -227,10 +227,7 @@ bool FObjectMixerEditorMainPanel::RequestAddObjectsToCollection(const FName& Col
 		for (const TSubclassOf<UObjectMixerObjectFilter>& Class : GetObjectFilterClasses())
 		{
 			const FName FilterName = Class->GetFName();
-			if (SerializedData->DoesCollectionExist(FilterName, CollectionName))
-			{
-				return SerializedData->AddObjectsToCollection(FilterName, CollectionName, ObjectsToAdd);
-			}
+			return SerializedData->AddObjectsToCollection(FilterName, CollectionName, ObjectsToAdd);
 		}
 	}
 
diff --git a/Engine/Plugins/Editor/ObjectMixer/ObjectMixer/Source/ObjectMixer/Private/Views/MainPanel/SObjectMixerEditorMainPanel.cpp b/Engine/Plugins/Editor/ObjectMixer/ObjectMixer/Source/ObjectMixer/Private/Views/MainPanel/SObjectMixerEditorMainPanel.cpp
--- a/Engine/Plugins/Editor/ObjectMixer/ObjectMixer/Source/ObjectMixer/Private/Views/MainPanel/SObjectMixerEditorMainPanel.cpp
+++ b/Engine/Plugins/Editor/ObjectMixer/ObjectMixer/Source/ObjectMixer/Private/Views/MainPanel/SObjectMixerEditorMainPanel.cpp
@@ -646,9 +646,9 @@ void SObjectMixerEditorMainPanel::SetSingleCollectionSelection(const FName& Coll
 		if (const TSharedPtr<FObjectMixerEditorListFilter_Collection> CollectionFilter =
 			StaticCastSharedRef<FObjectMixerEditorListFilter_Collection>(Filter))
 		{
-			const bool bIsAllCollectionFilter =
+			const bool bMatch =
 				CollectionFilter->CollectionName.IsEqual(CollectionToEnableName);
-			CollectionFilter->SetFilterActive(bIsAllCollectionFilter);
+			CollectionFilter->SetFilterActive(bMatch);
 		}
 	}
 }
@@ -956,6 +956,8 @@ bool SObjectMixerEditorMainPanel::RequestRemoveCollection(const FName& Collectio
 		
 		ListFilters.SetNum(NewListFilterCount);
 
+		OnCollectionCheckedStateChanged(true, UObjectMixerEditorSerializedData::AllCollectionName);
+
 		return true;
 	}
 
@@ -971,38 +973,7 @@ bool SObjectMixerEditorMainPanel::RequestRenameCollection(
 	const FName& CollectionNameToRename,
 	const FName& NewCollectionName)
 {
-	if (MainPanelModel.Pin()->RequestRenameCollection(CollectionNameToRename, NewCollectionName))
-	{
-		if (TSharedRef<IObjectMixerEditorListFilter>* Match = Algo::FindByPredicate(ListFilters,
-			[CollectionNameToRename](const TSharedRef<IObjectMixerEditorListFilter>& ListFilter)
-			{
-				if (const TSharedPtr<FObjectMixerEditorListFilter_Collection> CollectionFilter =
-					StaticCastSharedRef<FObjectMixerEditorListFilter_Collection>(ListFilter))
-				{
-					return CollectionFilter->CollectionName.IsEqual(CollectionNameToRename);
-				}
-				return false;
-			}))
-		{
-			StaticCastSharedRef<FObjectMixerEditorListFilter_Collection>(*Match)->CollectionName = NewCollectionName;
-		}
-		else
-		{
-			TSharedRef<FObjectMixerEditorListFilter_Collection> NewCollectionFilter =
-				MakeShared<FObjectMixerEditorListFilter_Collection>(NewCollectionName);
-		
-			ListFilters.Add(NewCollectionFilter);
-		
-			CollectionSelectorBox->AddSlot()
-			[
-				SNew(SCollectionSelectionButton, SharedThis(this), NewCollectionFilter)
-			];
-		}
-
-		return true;
-	}
-
-	return false;
+	return MainPanelModel.Pin()->RequestRenameCollection(CollectionNameToRename, NewCollectionName);
 }
 
 bool SObjectMixerEditorMainPanel::DoesCollectionExist(const FName& CollectionName) const
