Object Mixer:

-Copy/Paste Properties

diff --git a/Engine/Plugins/Editor/ObjectMixer/ObjectMixer/Source/ObjectMixer/ObjectMixerEditor.build.cs b/Engine/Plugins/Editor/ObjectMixer/ObjectMixer/Source/ObjectMixer/ObjectMixerEditor.build.cs
--- a/Engine/Plugins/Editor/ObjectMixer/ObjectMixer/Source/ObjectMixer/ObjectMixerEditor.build.cs
+++ b/Engine/Plugins/Editor/ObjectMixer/ObjectMixer/Source/ObjectMixer/ObjectMixerEditor.build.cs
@@ -29,6 +29,7 @@ public class ObjectMixerEditor : ModuleRules
 		PrivateDependencyModuleNames.AddRange(
 			new string[]
 			{
+				"ApplicationCore",
 				"AssetRegistry",
 				"AssetTools",
 				"CoreUObject",
diff --git a/Engine/Plugins/Editor/ObjectMixer/ObjectMixer/Source/ObjectMixer/Private/Views/List/SObjectMixerEditorListRow.cpp b/Engine/Plugins/Editor/ObjectMixer/ObjectMixer/Source/ObjectMixer/Private/Views/List/SObjectMixerEditorListRow.cpp
--- a/Engine/Plugins/Editor/ObjectMixer/ObjectMixer/Source/ObjectMixer/Private/Views/List/SObjectMixerEditorListRow.cpp
+++ b/Engine/Plugins/Editor/ObjectMixer/ObjectMixer/Source/ObjectMixer/Private/Views/List/SObjectMixerEditorListRow.cpp
@@ -11,7 +11,9 @@
 #include "Editor.h"
 #include "EditorActorFolders.h"
 #include "Engine/Blueprint.h"
+#include "Framework/MultiBox/MultiBoxBuilder.h"
 #include "GameFramework/Actor.h"
+#include "HAL/PlatformApplicationMisc.h"
 #include "Input/DragAndDrop.h"
 #include "ISinglePropertyView.h"
 #include "ObjectMixerEditorLog.h"
@@ -437,6 +439,183 @@ private:
 	float TextBlockLeftPadding = 3.0f;
 };
 
+class SInlinePropertyCellWidget : public SCompoundWidget
+{
+public:
+	
+	SLATE_BEGIN_ARGS(SInlinePropertyCellWidget)
+	{}
+
+	SLATE_EVENT(TDelegate<void(const FPropertyChangedEvent&)>, OnPropertyValueChanged)
+
+	SLATE_END_ARGS()
+	
+	void Construct(const FArguments& InArgs, UObject* ObjectRef, const FName InColumnName, const TSharedRef<FObjectMixerEditorListRow> RowPtr)
+	{
+		if (!ObjectRef || InColumnName.IsEqual(NAME_None))
+		{
+			return;
+		}
+
+		ColumnName = InColumnName;
+		WeakRowPtr = RowPtr;
+		
+		FPropertyEditorModule& PropertyEditorModule = FModuleManager::LoadModuleChecked<FPropertyEditorModule>("PropertyEditor");
+		FSinglePropertyParams Params;
+		Params.NamePlacement = EPropertyNamePlacement::Hidden;
+		
+		const TSharedPtr<ISinglePropertyView> SinglePropertyView =
+			PropertyEditorModule.CreateSingleProperty(ObjectRef, InColumnName, Params
+		);
+
+		if (SinglePropertyView)
+		{
+			if (const TSharedPtr<IPropertyHandle> Handle = SinglePropertyView->GetPropertyHandle())
+			{
+				if (const FProperty* Property = Handle->GetProperty())
+				{
+					RowPtr->PropertyNamesToHandles.Add(InColumnName, Handle);
+					
+					// Simultaneously edit all selected rows with a similar property
+					if (InArgs._OnPropertyValueChanged.IsBound())
+					{
+						Handle->SetOnPropertyValueChangedWithData(InArgs._OnPropertyValueChanged);
+						Handle->SetOnChildPropertyValueChangedWithData(InArgs._OnPropertyValueChanged);
+					}
+
+					ChildSlot
+					[
+						SNew(SBox)
+						.Visibility(EVisibility::Visible)
+						.HAlign(HAlign_Fill)
+						.VAlign(VAlign_Center)
+						[
+							SinglePropertyView.ToSharedRef()
+						]
+					];
+				}
+			}
+		}
+	}
+
+	virtual FReply OnMouseButtonUp(const FGeometry& MyGeometry, const FPointerEvent& MouseEvent) override
+	{
+		if (MouseEvent.GetEffectingButton() ==  EKeys::RightMouseButton)
+		{
+			return MakePropertyContextMenu(MouseEvent);
+		}
+		
+		return SCompoundWidget::OnMouseButtonUp(MyGeometry, MouseEvent);
+	}
+
+	FReply MakePropertyContextMenu(const FPointerEvent& MouseEvent)
+	{
+		FMenuBuilder MenuBuilder(true, nullptr, nullptr, true);
+
+		MenuBuilder.AddMenuEntry(
+			LOCTEXT("CopyProperty", "Copy"),
+			LOCTEXT("CopyProperty_ToolTip", "Copy this property value"),
+			FSlateIcon(FCoreStyle::Get().GetStyleSetName(), "GenericCommands.Copy"),
+			FExecuteAction::CreateRaw(this, &SInlinePropertyCellWidget::CopyPropertyValue));
+
+		FUIAction PasteAction = 
+			FUIAction(
+				FExecuteAction::CreateRaw(this, &SInlinePropertyCellWidget::PastePropertyValue),
+				FCanExecuteAction::CreateRaw(this, &SInlinePropertyCellWidget::CanPaste)
+			)
+		;
+
+		MenuBuilder.AddMenuEntry(
+			LOCTEXT("PasteProperty", "Paste"),
+			LOCTEXT("PasteProperty_ToolTip", "Paste the copied value here"),
+			FSlateIcon(FCoreStyle::Get().GetStyleSetName(), "GenericCommands.Paste"),
+			PasteAction);
+
+		MenuBuilder.AddMenuEntry(
+			LOCTEXT("CopyPropertyDisplayName", "Copy Display Name"),
+			LOCTEXT("CopyPropertyDisplayName_ToolTip", "Copy this property display name"),
+			FSlateIcon(FCoreStyle::Get().GetStyleSetName(), "GenericCommands.Copy"),
+			FExecuteAction::CreateRaw(this, &SInlinePropertyCellWidget::CopyPropertyName));
+
+		const FWidgetPath WidgetPath = MouseEvent.GetEventPath() != nullptr ? *MouseEvent.GetEventPath() : FWidgetPath();
+
+		FSlateApplication::Get().PushMenu(
+			AsShared(), WidgetPath, MenuBuilder.MakeWidget(),
+			MouseEvent.GetScreenSpacePosition(), FPopupTransitionEffect::ContextMenu
+		);
+
+		return FReply::Handled();
+	}
+
+protected:
+
+	void CopyPropertyValue() const
+	{
+		check(WeakRowPtr.IsValid());
+				
+		const TWeakPtr<IPropertyHandle>* PropertyHandle = WeakRowPtr.Pin()->PropertyNamesToHandles.Find(ColumnName);
+		if (PropertyHandle && PropertyHandle->IsValid())
+		{
+			FString Value;
+			if ((*PropertyHandle).Pin()->GetValueAsFormattedString(Value, PPF_Copy) == FPropertyAccess::Success)
+			{
+				FPlatformApplicationMisc::ClipboardCopy(*Value);
+			}
+		}
+	}
+
+	void PastePropertyValue() const
+	{
+		check(WeakRowPtr.IsValid());
+				
+		FString ClipboardContent;
+		FPlatformApplicationMisc::ClipboardPaste(ClipboardContent);
+				
+		const TWeakPtr<IPropertyHandle>* PropertyHandle = WeakRowPtr.Pin()->PropertyNamesToHandles.Find(ColumnName);
+		if (!ClipboardContent.IsEmpty() && PropertyHandle && PropertyHandle->IsValid())
+		{
+			(*PropertyHandle).Pin()->SetValueFromFormattedString(ClipboardContent, EPropertyValueSetFlags::InstanceObjects);
+		}
+	}
+
+	bool CanPaste() const
+	{
+		check(WeakRowPtr.IsValid());
+				
+		FString ClipboardContent;
+				
+		const TWeakPtr<IPropertyHandle>* PropertyHandle = WeakRowPtr.Pin()->PropertyNamesToHandles.Find(ColumnName);
+		if (PropertyHandle && PropertyHandle->IsValid())
+		{
+			if ((*PropertyHandle).Pin()->IsEditConst())
+			{
+				return false;
+			}
+		}
+
+		FPlatformApplicationMisc::ClipboardPaste(ClipboardContent);
+		return !ClipboardContent.IsEmpty();
+	}
+
+	void CopyPropertyName() const
+	{
+		check(WeakRowPtr.IsValid());
+				
+		const TWeakPtr<IPropertyHandle>* PropertyHandle = WeakRowPtr.Pin()->PropertyNamesToHandles.Find(ColumnName);
+		if (PropertyHandle && PropertyHandle->IsValid())
+		{
+			const FString Value = (*PropertyHandle).Pin()->GetPropertyDisplayName().ToString();
+			if (!Value.IsEmpty())
+			{
+				FPlatformApplicationMisc::ClipboardCopy(*Value);
+			}
+		}
+	}
+
+	TWeakPtr<FObjectMixerEditorListRow> WeakRowPtr;
+	FName ColumnName;
+};
+
 void SObjectMixerEditorListRow::Construct(
 	const FArguments& InArgs, const TSharedRef<STableViewBase>& InOwnerTable,
 	const TWeakPtr<FObjectMixerEditorListRow> InRow)
@@ -940,44 +1119,9 @@ TSharedPtr<SWidget> SObjectMixerEditorListRow::GenerateCells(
 	
 	if (UObject* ObjectRef = RowPtr->GetObject())
 	{
-		FPropertyEditorModule& PropertyEditorModule = FModuleManager::LoadModuleChecked<FPropertyEditorModule>("PropertyEditor");
-		{
-			FSinglePropertyParams Params;
-			Params.NamePlacement = EPropertyNamePlacement::Hidden;
-			
-			const TSharedPtr<ISinglePropertyView> SinglePropertyView =
-				PropertyEditorModule.CreateSingleProperty(ObjectRef, InColumnName, Params
-			);
-
-			if (SinglePropertyView)
-			{
-				if (const TSharedPtr<IPropertyHandle> Handle = SinglePropertyView->GetPropertyHandle())
-				{
-					if (const FProperty* Property = Handle->GetProperty())
-					{
-						const FName PropertyName = Property->GetFName();
-						RowPtr->PropertyNamesToHandles.Add(PropertyName, Handle);
-						
-						// Simultaneously edit all selected rows with a similar property
-						const TDelegate<void(const FPropertyChangedEvent&)> OnPropertyValueChanged =
-							TDelegate<void(const FPropertyChangedEvent&)>::CreateRaw(
-								this,
-								&SObjectMixerEditorListRow::OnPropertyChanged, PropertyName);
-					
-						Handle->SetOnPropertyValueChangedWithData(OnPropertyValueChanged);
-						Handle->SetOnChildPropertyValueChangedWithData(OnPropertyValueChanged);
-
-						return SNew(SBox)
-								.Visibility(EVisibility::SelfHitTestInvisible)
-								.HAlign(HAlign_Fill)
-								.VAlign(VAlign_Center)
-								[
-									SinglePropertyView.ToSharedRef()
-								];
-					}
-				}
-			}
-		}
+		return SNew(SInlinePropertyCellWidget, ObjectRef, InColumnName, RowPtr.ToSharedRef())
+				.OnPropertyValueChanged(this, &SObjectMixerEditorListRow::OnPropertyChanged, InColumnName)
+		;
 	}
 
 	return nullptr;
