Unshelved from pending changelist '23770228':

Object Mixer:

-New: Refactor the panel from basic STreeView to SSceneOutliner
-Fix: Cannot drag light sources to unparent them in the Light Mixer Window
-Fix: Unable Drag Drop Light From Light Mixer to Sequencer
-Fix: Disable Ability To Rename Lights contained in Blueprints or other Actors
-Fix: Rebuild List Option not working
-New: Add Additional Column Info options
-New: Ability to reparent light to root by drag and drop
-New: Work with arbitrary worlds
-Fix: Remove proprietary list filter implementation

diff --git a/Engine/Plugins/Editor/ObjectMixer/ObjectMixer/Source/ObjectMixer/ObjectMixerEditor.build.cs b/Engine/Plugins/Editor/ObjectMixer/ObjectMixer/Source/ObjectMixer/ObjectMixerEditor.build.cs
--- a/Engine/Plugins/Editor/ObjectMixer/ObjectMixer/Source/ObjectMixer/ObjectMixerEditor.build.cs
+++ b/Engine/Plugins/Editor/ObjectMixer/ObjectMixer/Source/ObjectMixer/ObjectMixerEditor.build.cs
@@ -35,6 +35,7 @@ public class ObjectMixerEditor : ModuleRules
 				"CoreUObject",
 				"ContentBrowser",
 				"Engine",
+				"EditorConfig",
 				"EditorStyle",
 				"EditorWidgets",
 				"InputCore",
diff --git a/Engine/Plugins/Editor/ObjectMixer/ObjectMixer/Source/ObjectMixer/Private/ObjectMixerEditorModule.cpp b/Engine/Plugins/Editor/ObjectMixer/ObjectMixer/Source/ObjectMixer/Private/ObjectMixerEditorModule.cpp
--- a/Engine/Plugins/Editor/ObjectMixer/ObjectMixer/Source/ObjectMixer/Private/ObjectMixerEditorModule.cpp
+++ b/Engine/Plugins/Editor/ObjectMixer/ObjectMixer/Source/ObjectMixer/Private/ObjectMixerEditorModule.cpp
@@ -2,26 +2,27 @@
 
 #include "ObjectMixerEditorModule.h"
 
+#include "ObjectMixerEditorLog.h"
 #include "Engine/Level.h"
 #include "ObjectMixerEditorSettings.h"
 #include "Misc/CoreDelegates.h"
+#include "ObjectMixerEditorSerializedData.h"
 #include "ObjectMixerEditorStyle.h"
-#include "Views/MainPanel/ObjectMixerEditorMainPanel.h"
+#include "Views/List/ObjectMixerEditorList.h"
+#include "Views/Widgets/ObjectMixerEditorListMenuContext.h"
 
 #include "ISettingsModule.h"
 #include "LevelEditor.h"
-#include "ObjectMixerEditorSerializedData.h"
 #include "Selection.h"
 #include "Misc/TransactionObjectEvent.h"
 #include "ToolMenus.h"
 #include "Widgets/Docking/SDockTab.h"
 #include "WorkspaceMenuStructure.h"
 #include "WorkspaceMenuStructureModule.h"
-#include "Views/Widgets/ObjectMixerEditorListMenuContext.h"
 
 const FName FObjectMixerEditorModule::BaseObjectMixerModuleName("ObjectMixerEditor");
 
-#define LOCTEXT_NAMESPACE "FObjectMixerEditorModule"
+#define LOCTEXT_NAMESPACE "ObjectMixerEditor"
 
 IMPLEMENT_MODULE(FObjectMixerEditorModule, ObjectMixerEditor)
 
@@ -98,7 +99,7 @@ void FObjectMixerEditorModule::Teardown()
 	}
 	DelegateHandles.Empty();
 	
-	MainPanel.Reset();
+	ListModel.Reset();
 
 	UToolMenus::UnregisterOwner(this);
 	
@@ -125,51 +126,43 @@ FName FObjectMixerEditorModule::GetModuleName() const
 TSharedPtr<SWidget> FObjectMixerEditorModule::MakeObjectMixerDialog(
 	TSubclassOf<UObjectMixerObjectFilter> InDefaultFilterClass)
 {
-	if (!MainPanel.IsValid())
+	if (!ListModel.IsValid())
 	{
-		MainPanel = MakeShared<FObjectMixerEditorMainPanel>(GetModuleName());
-		MainPanel->Initialize();
+		ListModel = MakeShared<FObjectMixerEditorList>(GetModuleName());
+		ListModel->Initialize();
 	}
-	
-	const TSharedPtr<SWidget> ObjectMixerDialog = MainPanel->GetOrCreateWidget();
 
 	if (InDefaultFilterClass)
 	{
-		MainPanel->SetDefaultFilterClass(InDefaultFilterClass);
+		ListModel->SetDefaultFilterClass(InDefaultFilterClass);
 	}
+	
+	const TSharedPtr<SWidget> ObjectMixerDialog = ListModel->GetOrCreateWidget();
 
 	return ObjectMixerDialog;
 }
 
 void FObjectMixerEditorModule::RequestRebuildList() const
 {
-	if (MainPanel.IsValid())
+	if (ListModel.IsValid())
 	{
-		MainPanel->RequestRebuildList();
+		ListModel->RequestRebuildList();
 	}
 }
 
 void FObjectMixerEditorModule::RefreshList() const
 {
-	if (MainPanel.IsValid())
+	if (ListModel.IsValid())
 	{
-		MainPanel->RefreshList();
+		ListModel->RefreshList();
 	}
 }
 
 void FObjectMixerEditorModule::OnRenameCommand()
 {
-	if (MainPanel.IsValid())
+	if (ListModel.IsValid())
 	{
-		MainPanel->OnRenameCommand();
-	}
-}
-
-void FObjectMixerEditorModule::RequestSyncEditorSelectionToListSelection()
-{
-	if (MainPanel.IsValid())
-	{
-		MainPanel->RequestSyncEditorSelectionToListSelection();
+		ListModel->OnRenameCommand();
 	}
 }
 
@@ -329,10 +322,6 @@ void FObjectMixerEditorModule::BindDelegates()
 	{
 		RequestRebuildList();
 	}));
-	DelegateHandles.Add(GEditor->GetSelectedActors()->SelectionChangedEvent.AddLambda([this] (UObject*)
-	{
-		RequestSyncEditorSelectionToListSelection();
-	}));
 	
 	DelegateHandles.Add(FEditorDelegates::MapChange.AddLambda([this](uint32)
 	{
@@ -378,9 +367,9 @@ TSet<FName> FObjectMixerEditorModule::GetPropertiesThatRequireRefresh() const
 {
 	TSet<FName> ReturnValue;
 
-	if (MainPanel.IsValid())
+	if (ListModel.IsValid())
 	{
-		for (const TObjectPtr<UObjectMixerObjectFilter> Instance : MainPanel->GetObjectFilterInstances())
+		for (const TObjectPtr<UObjectMixerObjectFilter> Instance : ListModel->GetObjectFilterInstances())
 		{
 			ReturnValue.Append(Instance->GetPropertiesThatRequireListRefresh());
 		}
diff --git a/Engine/Plugins/Editor/ObjectMixer/ObjectMixer/Source/ObjectMixer/Private/ObjectMixerFilterFactory.cpp b/Engine/Plugins/Editor/ObjectMixer/ObjectMixer/Source/ObjectMixer/Private/ObjectMixerFilterFactory.cpp
--- a/Engine/Plugins/Editor/ObjectMixer/ObjectMixer/Source/ObjectMixer/Private/ObjectMixerFilterFactory.cpp
+++ b/Engine/Plugins/Editor/ObjectMixer/ObjectMixer/Source/ObjectMixer/Private/ObjectMixerFilterFactory.cpp
@@ -9,7 +9,7 @@
 #include "Kismet2/KismetEditorUtilities.h"
 #include "Engine/BlueprintGeneratedClass.h"
 
-#define LOCTEXT_NAMESPACE "ObjectMixerFilterFactory"
+#define LOCTEXT_NAMESPACE "ObjectMixerEditor"
 
 UObjectMixerBlueprintFilterFactory::UObjectMixerBlueprintFilterFactory()
 {
diff --git a/Engine/Plugins/Editor/ObjectMixer/ObjectMixer/Source/ObjectMixer/Private/Views/List/ObjectMixerEditorList.cpp b/Engine/Plugins/Editor/ObjectMixer/ObjectMixer/Source/ObjectMixer/Private/Views/List/ObjectMixerEditorList.cpp
--- a/Engine/Plugins/Editor/ObjectMixer/ObjectMixer/Source/ObjectMixer/Private/Views/List/ObjectMixerEditorList.cpp
+++ b/Engine/Plugins/Editor/ObjectMixer/ObjectMixer/Source/ObjectMixer/Private/Views/List/ObjectMixerEditorList.cpp
@@ -2,17 +2,318 @@
 
 #include "Views/List/ObjectMixerEditorList.h"
 
+#include "ObjectMixerEditorLog.h"
+#include "ObjectMixerEditorSerializedData.h"
 #include "Views/List/SObjectMixerEditorList.h"
-#include "Views/MainPanel/ObjectMixerEditorMainPanel.h"
 
-FObjectMixerEditorList::FObjectMixerEditorList(TSharedRef<FObjectMixerEditorMainPanel, ESPMode::ThreadSafe> InMainPanel)
+#include "LevelEditorActions.h"
+#include "Framework/Commands/GenericCommands.h"
+
+void FObjectMixerEditorList::Initialize()
 {
-	MainPanelModelPtr = InMainPanel;
-
-	InMainPanel->OnPreFilterChange.AddRaw(this, &FObjectMixerEditorList::OnPreFilterChange);
-	InMainPanel->OnPostFilterChange.AddRaw(this, &FObjectMixerEditorList::OnPostFilterChange);
+	RegisterAndMapContextMenuCommands();
 }
 
+void FObjectMixerEditorList::RegisterAndMapContextMenuCommands()
+{
+	ObjectMixerElementEditCommands = MakeShared<FUICommandList>();
+
+	ObjectMixerElementEditCommands->MapAction(FGenericCommands::Get().Cut,
+		FExecuteAction::CreateStatic( &FLevelEditorActionCallbacks::ExecuteExecCommand, FString( TEXT("EDIT CUT") ) ),
+		FCanExecuteAction::CreateStatic( &FLevelEditorActionCallbacks::Cut_CanExecute )
+	);
+	ObjectMixerElementEditCommands->MapAction(FGenericCommands::Get().Copy,
+		FExecuteAction::CreateStatic( &FLevelEditorActionCallbacks::ExecuteExecCommand, FString( TEXT("EDIT COPY") ) ),
+		FCanExecuteAction::CreateStatic( &FLevelEditorActionCallbacks::Copy_CanExecute )
+	);
+	ObjectMixerElementEditCommands->MapAction(FGenericCommands::Get().Paste,
+		FExecuteAction::CreateStatic( &FLevelEditorActionCallbacks::ExecuteExecCommand, FString( TEXT("EDIT PASTE") ) ),
+		FCanExecuteAction::CreateStatic( &FLevelEditorActionCallbacks::Paste_CanExecute )
+	);
+	ObjectMixerElementEditCommands->MapAction(FGenericCommands::Get().Duplicate,
+		FExecuteAction::CreateStatic( &FLevelEditorActionCallbacks::ExecuteExecCommand, FString( TEXT("DUPLICATE") ) ),
+		FCanExecuteAction::CreateStatic( &FLevelEditorActionCallbacks::Duplicate_CanExecute )
+	);
+	ObjectMixerElementEditCommands->MapAction(FGenericCommands::Get().Delete,
+		FExecuteAction::CreateStatic( &FLevelEditorActionCallbacks::ExecuteExecCommand, FString( TEXT("DELETE") ) ),
+		FCanExecuteAction::CreateStatic( &FLevelEditorActionCallbacks::Delete_CanExecute )
+	);
+	ObjectMixerElementEditCommands->MapAction(FGenericCommands::Get().Rename,
+		FUIAction(FExecuteAction::CreateRaw(this, &FObjectMixerEditorList::OnRenameCommand))
+	);
+}
+
+void FObjectMixerEditorList::RebuildCollectionSelector()
+{
+	ListWidget->RebuildCollectionSelector();
+}
+
+void FObjectMixerEditorList::SetDefaultFilterClass(UClass* InNewClass)
+{
+	DefaultFilterClass = InNewClass;
+	AddObjectFilterClass(InNewClass, false);
+}
+
+bool FObjectMixerEditorList::IsClassSelected(UClass* InClass) const
+{
+	return GetObjectFilterClasses().Contains(InClass);
+}
+
+const TArray<TObjectPtr<UObjectMixerObjectFilter>>& FObjectMixerEditorList::GetObjectFilterInstances()
+{
+	if (ObjectFilterInstances.Num() == 0)
+	{
+		CacheObjectFilterInstances();
+	}
+
+	return ObjectFilterInstances;
+}
+
+const UObjectMixerObjectFilter* FObjectMixerEditorList::GetMainObjectFilterInstance()
+{
+	return GetObjectFilterInstances().Num() > 0 ? GetObjectFilterInstances()[0].Get() : nullptr;
+}
+
+void FObjectMixerEditorList::CacheObjectFilterInstances()
+{
+	ObjectFilterInstances.Reset();
+	
+	for (const TSubclassOf<UObjectMixerObjectFilter> Class : GetObjectFilterClasses())
+	{
+		UObjectMixerObjectFilter* NewInstance = NewObject<UObjectMixerObjectFilter>(GetTransientPackage(), Class);
+		ObjectFilterInstances.Add(TObjectPtr<UObjectMixerObjectFilter>(NewInstance));
+	}
+
+	BuildPerformanceCache();
+}
+
+TSet<UClass*> FObjectMixerEditorList::ForceGetObjectClassesToFilter()
+{
+	TSet<UClass*> ReturnValue;
+	for (const TObjectPtr<UObjectMixerObjectFilter>& Filter : GetObjectFilterInstances())
+	{
+		ReturnValue.Append(Filter->GetObjectClassesToFilter());
+	}
+		
+	return ReturnValue;
+}
+
+TSet<TSubclassOf<AActor>> FObjectMixerEditorList::ForceGetObjectClassesToPlace()
+{
+	TSet<TSubclassOf<AActor>> ReturnValue;
+
+	for (const TObjectPtr<UObjectMixerObjectFilter>& Filter : GetObjectFilterInstances())
+	{
+		ReturnValue.Append(Filter->GetObjectClassesToPlace());
+	}
+		
+	return ReturnValue;
+}
+
+void FObjectMixerEditorList::AddObjectFilterClass(UClass* InObjectFilterClass, const bool bShouldRebuild)
+{
+	if (ensureAlwaysMsgf(InObjectFilterClass->IsChildOf(UObjectMixerObjectFilter::StaticClass()), TEXT("%hs: Class '%s' is not a child of UObjectMixerObjectFilter."), __FUNCTION__, *InObjectFilterClass->GetName()))
+	{
+		ObjectFilterClasses.Add(InObjectFilterClass);
+
+		CacheObjectFilterInstances();
+
+		if (bShouldRebuild)
+		{
+			RequestRebuildList();
+		}
+	}
+}
+
+void FObjectMixerEditorList::RemoveObjectFilterClass(UClass* InObjectFilterClass, const bool bCacheAndRebuild)
+{
+	if (ObjectFilterClasses.Remove(InObjectFilterClass) > 0 && bCacheAndRebuild)
+	{
+		CacheAndRebuildFilters();
+	}
+}
+
+UObjectMixerEditorSerializedData* FObjectMixerEditorList::GetSerializedData() const
+{
+	return GetMutableDefault<UObjectMixerEditorSerializedData>();
+}
+
+bool FObjectMixerEditorList::RequestAddObjectsToCollection(const FName& CollectionName, const TSet<FSoftObjectPath>& ObjectsToAdd) const
+{
+	if (UObjectMixerEditorSerializedData* SerializedData = GetSerializedData())
+	{
+		for (const TSubclassOf<UObjectMixerObjectFilter>& Class : GetObjectFilterClasses())
+		{
+			const FName FilterName = Class->GetFName();
+			return SerializedData->AddObjectsToCollection(FilterName, CollectionName, ObjectsToAdd);
+		}
+	}
+
+	return false;
+}
+
+bool FObjectMixerEditorList::RequestRemoveObjectsFromCollection(const FName& CollectionName, const TSet<FSoftObjectPath>& ObjectsToRemove) const
+{
+	if (UObjectMixerEditorSerializedData* SerializedData = GetSerializedData())
+	{
+		for (const TSubclassOf<UObjectMixerObjectFilter>& Class : GetObjectFilterClasses())
+		{
+			const FName FilterName = Class->GetFName();
+			if (SerializedData->DoesCollectionExist(FilterName, CollectionName))
+			{
+				return SerializedData->RemoveObjectsFromCollection(FilterName, CollectionName, ObjectsToRemove);
+			}
+		}
+	}
+
+	return false;
+}
+
+bool FObjectMixerEditorList::RequestRemoveCollection(const FName& CollectionName) const
+{
+	if (UObjectMixerEditorSerializedData* SerializedData = GetSerializedData())
+	{
+		for (const TSubclassOf<UObjectMixerObjectFilter>& Class : GetObjectFilterClasses())
+		{
+			const FName FilterName = Class->GetFName();
+			if (SerializedData->DoesCollectionExist(FilterName, CollectionName))
+			{
+				return SerializedData->RemoveCollection(FilterName, CollectionName);
+			}
+		}
+	}
+
+	return false;
+}
+
+bool FObjectMixerEditorList::RequestDuplicateCollection(const FName& CollectionToDuplicateName, FName& DesiredDuplicateName) const
+{
+	if (UObjectMixerEditorSerializedData* SerializedData = GetSerializedData())
+	{
+		for (const TSubclassOf<UObjectMixerObjectFilter>& Class : GetObjectFilterClasses())
+		{
+			const FName FilterName = Class->GetFName();
+			if (SerializedData->DoesCollectionExist(FilterName, CollectionToDuplicateName))
+			{
+				return SerializedData->DuplicateCollection(FilterName, CollectionToDuplicateName, DesiredDuplicateName);
+			}
+		}
+	}
+
+	return false;
+}
+
+bool FObjectMixerEditorList::RequestReorderCollection(const FName& CollectionToMoveName, const FName& CollectionInsertBeforeName) const
+{
+	if (UObjectMixerEditorSerializedData* SerializedData = GetSerializedData())
+	{
+		for (const TSubclassOf<UObjectMixerObjectFilter>& Class : GetObjectFilterClasses())
+		{
+			const FName FilterName = Class->GetFName();
+			if (SerializedData->DoesCollectionExist(FilterName, CollectionToMoveName))
+			{
+				return SerializedData->ReorderCollection(FilterName, CollectionToMoveName, CollectionInsertBeforeName);
+			}
+		}
+	}
+
+	return false;
+}
+
+bool FObjectMixerEditorList::RequestRenameCollection(const FName& CollectionNameToRename,
+	const FName& NewCollectionName) const
+{
+	if (UObjectMixerEditorSerializedData* SerializedData = GetSerializedData())
+	{
+		for (const TSubclassOf<UObjectMixerObjectFilter>& Class : GetObjectFilterClasses())
+		{
+			const FName FilterName = Class->GetFName();
+			if (SerializedData->DoesCollectionExist(FilterName, CollectionNameToRename))
+			{
+				return SerializedData->RenameCollection(FilterName, CollectionNameToRename, NewCollectionName);
+			}
+		}
+	}
+
+	return false;
+}
+
+bool FObjectMixerEditorList::DoesCollectionExist(const FName& CollectionName) const
+{
+	if (UObjectMixerEditorSerializedData* SerializedData = GetSerializedData())
+	{
+		for (const TSubclassOf<UObjectMixerObjectFilter>& Class : GetObjectFilterClasses())
+		{
+			const FName FilterName = Class->GetFName();
+			if (SerializedData->DoesCollectionExist(FilterName, CollectionName))
+			{
+				return true;
+			}
+		}
+	}
+
+	return false;
+}
+
+bool FObjectMixerEditorList::IsObjectInCollection(const FName& CollectionName, const FSoftObjectPath& InObject) const
+{
+	if (UObjectMixerEditorSerializedData* SerializedData = GetSerializedData())
+	{
+		for (const TSubclassOf<UObjectMixerObjectFilter>& Class : GetObjectFilterClasses())
+		{
+			const FName FilterName = Class->GetFName();
+			if (SerializedData->DoesCollectionExist(FilterName, CollectionName))
+			{
+				return SerializedData->IsObjectInCollection(FilterName, CollectionName, InObject);
+			}
+		}
+	}
+
+	return false;
+}
+
+TSet<FName> FObjectMixerEditorList::GetCollectionsForObject(const FSoftObjectPath& InObject) const
+{
+	TSet<FName> ReturnValue;
+	if (UObjectMixerEditorSerializedData* SerializedData = GetSerializedData())
+	{
+		for (const TSubclassOf<UObjectMixerObjectFilter>& Class : GetObjectFilterClasses())
+		{
+			const FName FilterName = Class->GetFName();
+			ReturnValue.Append(SerializedData->GetCollectionsForObject(FilterName, InObject));
+		}
+	}
+
+	return ReturnValue;
+}
+
+TArray<FName> FObjectMixerEditorList::GetAllCollectionNames() const
+{
+	TArray<FName> ReturnValue;
+	if (UObjectMixerEditorSerializedData* SerializedData = GetSerializedData())
+	{
+		for (const TSubclassOf<UObjectMixerObjectFilter>& Class : GetObjectFilterClasses())
+		{
+			const FName FilterName = Class->GetFName();
+			ReturnValue.Append(SerializedData->GetAllCollectionNames(FilterName));
+		}
+	}
+
+	return ReturnValue;
+}
+
+const TSubclassOf<UObjectMixerObjectFilter>& FObjectMixerEditorList::GetDefaultFilterClass() const
+{
+	return DefaultFilterClass;
+}
+
+FObjectMixerEditorList::FObjectMixerEditorList(const FName InModuleName)
+{
+	ModuleName = InModuleName;
+	
+	OnPostFilterChangeDelegate.AddRaw(this, &FObjectMixerEditorList::OnPostFilterChange);
+}
+
 FObjectMixerEditorList::~FObjectMixerEditorList()
 {
 	FlushWidget();
@@ -20,11 +321,8 @@
 
 void FObjectMixerEditorList::FlushWidget()
 {
-	if (TSharedPtr<FObjectMixerEditorMainPanel> MainPanelPinned = MainPanelModelPtr.Pin())
-	{
-		MainPanelPinned->OnPreFilterChange.RemoveAll(this);
-		MainPanelPinned->OnPostFilterChange.RemoveAll(this);
-	}
+	OnPreFilterChangeDelegate.RemoveAll(this);
+	OnPostFilterChangeDelegate.RemoveAll(this);
 	
 	ListWidget.Reset();
 }
@@ -41,37 +339,11 @@
 	return ListWidget.ToSharedRef();
 }
 
-void FObjectMixerEditorList::OnPreFilterChange()
-{
-	if (ListWidget.IsValid())
-	{
-		if (const TSharedPtr<FObjectMixerEditorMainPanel> PinnedMainPanel = GetMainPanelModel().Pin())
-		{		
-			ListWidget->CacheTreeState(PinnedMainPanel->GetWeakActiveListFiltersSortedByName());
-		}
-	}
-}
-
 void FObjectMixerEditorList::OnPostFilterChange()
 {
-	if (ListWidget.IsValid())
-	{
-		if (const TSharedPtr<FObjectMixerEditorMainPanel> PinnedMainPanel = GetMainPanelModel().Pin())
-		{		
-			ListWidget->EvaluateIfRowsPassFilters();
-			ListWidget->RestoreTreeState(PinnedMainPanel->GetWeakActiveListFiltersSortedByName());
-		}
-	}
+	RequestRebuildList();
 }
 
-void FObjectMixerEditorList::ClearList() const
-{
-	if (ListWidget.IsValid())
-	{
-		ListWidget->ClearList();
-	}
-}
-
 void FObjectMixerEditorList::RequestRebuildList() const
 {
 	if (ListWidget.IsValid())
@@ -88,6 +360,72 @@
 	}
 }
 
+void FObjectMixerEditorList::BuildPerformanceCache()
+{
+	TSet<UClass*> LocalObjectClassesToFilterCache;
+	TSet<FName> LocalColumnsToShowByDefaultCache;
+	TSet<FName> LocalColumnsToExcludeCache;
+	TSet<FName> LocalForceAddedColumnsCache;
+	
+	for (const TObjectPtr<UObjectMixerObjectFilter>& FilterInstance : GetObjectFilterInstances())
+	{
+		if (!FilterInstance)
+		{
+			UE_LOG(LogObjectMixerEditor, Display, TEXT("%hs: UObjectMixerObjectFilter instance not valid."), __FUNCTION__);
+			continue;
+		}
+
+		LocalObjectClassesToFilterCache.Append(FilterInstance->GetObjectClassesToFilter());
+		LocalColumnsToShowByDefaultCache.Append(FilterInstance->GetColumnsToShowByDefault());
+		LocalColumnsToExcludeCache.Append(FilterInstance->GetColumnsToExclude());
+		LocalForceAddedColumnsCache.Append(FilterInstance->GetForceAddedColumns());
+	}
+
+	if (LocalObjectClassesToFilterCache.Difference(ObjectClassesToFilterCache).Num() > 0 || 
+		ObjectClassesToFilterCache.Difference(LocalObjectClassesToFilterCache).Num() > 0)
+	{
+		ObjectClassesToFilterCache = LocalObjectClassesToFilterCache;
+	}
+
+	if (LocalColumnsToShowByDefaultCache.Difference(ColumnsToShowByDefaultCache).Num() > 0 || 
+		ColumnsToShowByDefaultCache.Difference(LocalColumnsToShowByDefaultCache).Num() > 0)
+	{
+		ColumnsToShowByDefaultCache = LocalColumnsToShowByDefaultCache;
+	}
+
+	if (LocalColumnsToExcludeCache.Difference(ColumnsToExcludeCache).Num() > 0 || 
+		ColumnsToExcludeCache.Difference(LocalColumnsToExcludeCache).Num() > 0)
+	{
+		ColumnsToExcludeCache = LocalColumnsToExcludeCache;
+	}
+
+	if (LocalForceAddedColumnsCache.Difference(ForceAddedColumnsCache).Num() > 0 || 
+		ForceAddedColumnsCache.Difference(LocalForceAddedColumnsCache).Num() > 0)
+	{
+		ForceAddedColumnsCache = LocalForceAddedColumnsCache;
+	}
+
+	// These properties should be governed by the Main instance
+	if (const UObjectMixerObjectFilter* FilterInstance = GetMainObjectFilterInstance())
+	{
+		if (const EObjectMixerInheritanceInclusionOptions PropertyInheritanceInclusionOptions =
+			FilterInstance->GetObjectMixerPropertyInheritanceInclusionOptions(); 
+			PropertyInheritanceInclusionOptions != PropertyInheritanceInclusionOptionsCache)
+		{
+			PropertyInheritanceInclusionOptionsCache = PropertyInheritanceInclusionOptions;
+		}
+		if (const bool bShouldIncludeUnsupportedProperties = FilterInstance->ShouldIncludeUnsupportedProperties(); 
+			bShouldIncludeUnsupportedProperties != bShouldIncludeUnsupportedPropertiesCache)
+		{
+			bShouldIncludeUnsupportedPropertiesCache = bShouldIncludeUnsupportedProperties;
+		}
+	}
+	else
+	{
+		UE_LOG(LogObjectMixerEditor, Display, TEXT("%hs: No Main UObjectMixerObjectFilter instance found."), __FUNCTION__);
+	}
+}
+
 void FObjectMixerEditorList::OnRenameCommand()
 {
 	if (ListWidget.IsValid())
@@ -96,18 +434,10 @@
 	}
 }
 
-void FObjectMixerEditorList::RequestSyncEditorSelectionToListSelection() const
+TArray<TSharedPtr<ISceneOutlinerTreeItem>> FObjectMixerEditorList::GetSelectedTreeViewItems() const
 {
 	if (ListWidget.IsValid())
 	{
-		ListWidget->RequestSyncEditorSelectionToListSelection();
-	}
-}
-
-TArray<FObjectMixerEditorListRowPtr> FObjectMixerEditorList::GetSelectedTreeViewItems() const
-{
-	if (ListWidget.IsValid())
-	{
 		return ListWidget->GetSelectedTreeViewItems();
 	}
 
@@ -124,40 +454,7 @@
 	return INDEX_NONE;
 }
 
-void FObjectMixerEditorList::OnRequestNewFolder(TOptional<FFolder> ExplicitParentFolder)
+TSet<TSharedPtr<ISceneOutlinerTreeItem>> FObjectMixerEditorList::GetSoloRows() const
 {
 	if (ListWidget.IsValid())
 	{
-		ListWidget->OnRequestNewFolder(ExplicitParentFolder);
-	}
-}
-
-void FObjectMixerEditorList::OnRequestMoveFolder(const FFolder& FolderToMove, const FFolder& TargetNewParentFolder)
-{
-	if (ListWidget.IsValid())
-	{
-		ListWidget->OnRequestMoveFolder(FolderToMove, TargetNewParentFolder);
-	}
-}
-
-void FObjectMixerEditorList::ExecuteListViewSearchOnAllRows(const FString& SearchString,
-                                                            const bool bShouldRefreshAfterward)
-{
-	if (ListWidget.IsValid())
-	{
-		ListWidget->ExecuteListViewSearchOnAllRows(SearchString, bShouldRefreshAfterward);
-	}
-}
-
-void FObjectMixerEditorList::EvaluateIfRowsPassFilters(const bool bShouldRefreshAfterward) const
-{
-	if (ListWidget.IsValid())
-	{
-		ListWidget->EvaluateIfRowsPassFilters();
-	}
-}
-
-TSet<TWeakPtr<FObjectMixerEditorListRow>> FObjectMixerEditorList::GetSoloRows() const
-{
-	if (ListWidget.IsValid())
-	{
@@ -193,7 +490,41 @@
 	}
 }
 
-TWeakPtr<FObjectMixerEditorMainPanel> FObjectMixerEditorList::GetMainPanelModel()
+const TSet<FName>& FObjectMixerEditorList::GetSelectedCollections() const
 {
-	return MainPanelModelPtr;
+	return SelectedCollections;
 }
+
+bool FObjectMixerEditorList::IsCollectionSelected(const FName& CollectionName)
+{
+	return GetSelectedCollections().Contains(CollectionName);
+}
+
+void FObjectMixerEditorList::SetSelectedCollections(const TSet<FName> InSelectedCollections)
+{
+	SelectedCollections = InSelectedCollections;
+}
+
+void FObjectMixerEditorList::SetCollectionSelected(const FName& CollectionName, const bool bNewSelected)
+{
+	if (UObjectMixerEditorSerializedData* SerializedData = GetMutableDefault<UObjectMixerEditorSerializedData>())
+	{
+		for (const TSubclassOf<UObjectMixerObjectFilter>& Class : GetObjectFilterClasses())
+		{
+			const FName FilterName = Class->GetFName();
+			if (SerializedData->DoesCollectionExist(FilterName, CollectionName))
+			{
+				if (bNewSelected)
+				{
+					SelectedCollections.Add(CollectionName);
+				}
+				else
+				{
+					SelectedCollections.Remove(CollectionName);
+				}
+
+				return;
+			}
+		}
+	}
+}
diff --git a/Engine/Plugins/Editor/ObjectMixer/ObjectMixer/Source/ObjectMixer/Private/Views/List/ObjectMixerEditorListRowData.cpp b/Engine/Plugins/Editor/ObjectMixer/ObjectMixer/Source/ObjectMixer/Private/Views/List/ObjectMixerEditorListRowData.cpp
new file mode 100644
--- /dev/null
+++ b/Engine/Plugins/Editor/ObjectMixer/ObjectMixer/Source/ObjectMixer/Private/Views/List/ObjectMixerEditorListRowData.cpp
@@ -0,0 +1,300 @@
+// Copyright Epic Games, Inc. All Rights Reserved.
+
+#include "Views/List/ObjectMixerEditorListRowData.h"
+
+#include "ObjectMixerEditorLog.h"
+#include "Views/List/ObjectMixerUtils.h"
+#include "Views/List/SObjectMixerEditorList.h"
+
+#include "ScopedTransaction.h"
+#include "GameFramework/Actor.h"
+
+const TArray<TObjectPtr<UObjectMixerObjectFilter>>& FObjectMixerEditorListRowData::GetObjectFilterInstances() const
+{
+	SObjectMixerEditorList* ListView = GetListView();
+	check (ListView)
+
+	const TSharedPtr<FObjectMixerEditorList> PinnedListModel = ListView->GetListModelPtr().Pin();
+	check (PinnedListModel);
+	
+	return PinnedListModel->GetObjectFilterInstances();
+}
+
+const UObjectMixerObjectFilter* FObjectMixerEditorListRowData::GetMainObjectFilterInstance() const
+{
+	if (SObjectMixerEditorList* ListView = GetListView())
+	{
+		if (const TSharedPtr<FObjectMixerEditorList> PinnedListModel = ListView->GetListModelPtr().Pin())
+		{
+			return PinnedListModel->GetMainObjectFilterInstance();
+		}
+	}
+	
+	return nullptr;
+}
+
+bool FObjectMixerEditorListRowData::GetIsTreeViewItemExpanded(const TSharedRef<ISceneOutlinerTreeItem> InRow)
+{
+	return GetListView()->IsTreeViewItemExpanded(InRow);
+}
+
+void FObjectMixerEditorListRowData::SetIsTreeViewItemExpanded(const TSharedRef<ISceneOutlinerTreeItem> InRow, const bool bNewExpanded)
+{
+	GetListView()->SetTreeViewItemExpanded(InRow, bNewExpanded);
+}
+
+bool FObjectMixerEditorListRowData::GetDoesRowPassFilters() const
+{
+	return bDoesRowPassFilters;
+}
+
+void FObjectMixerEditorListRowData::SetDoesRowPassFilters(const bool bPass)
+{
+	bDoesRowPassFilters = bPass;
+}
+
+bool FObjectMixerEditorListRowData::GetIsSelected(const TSharedRef<ISceneOutlinerTreeItem> InRow)
+{
+	if (SObjectMixerEditorList* ListView = GetListView())
+	{
+		return ListView->IsTreeViewItemSelected(InRow);
+	}
+
+	return false;
+}
+
+void FObjectMixerEditorListRowData::SetIsSelected(const TSharedRef<ISceneOutlinerTreeItem> InRow, const bool bNewSelected)
+{
+	if (SObjectMixerEditorList* ListView = GetListView())
+	{
+		return ListView->SetTreeViewItemSelected(InRow, bNewSelected);
+	}
+}
+
+bool FObjectMixerEditorListRowData::HasAtLeastOneChildThatIsNotSolo(const TSharedRef<ISceneOutlinerTreeItem> InRow, const bool bRecursive) const
+{
+	for (const TWeakPtr<ISceneOutlinerTreeItem>& ChildRow : InRow->GetChildren())
+	{
+		if (const TSharedPtr<ISceneOutlinerTreeItem> PinnedChildRow = ChildRow.Pin())
+		{
+			if (!FObjectMixerUtils::GetRowData(PinnedChildRow)->GetRowSoloState())
+			{
+				return true;
+			}
+
+			if (bRecursive && FObjectMixerUtils::GetRowData(PinnedChildRow)->HasAtLeastOneChildThatIsNotSolo(PinnedChildRow.ToSharedRef(), true))
+			{
+				return true;
+			}
+		}
+	}
+
+	return false;
+}
+
+FText FObjectMixerEditorListRowData::GetDisplayName(TSharedPtr<ISceneOutlinerTreeItem> InTreeItem, const bool bIsHybridRow) const
+{
+	const FText Override = GetDisplayNameOverride();
+	if (!Override.IsEmpty())
+	{
+		return Override;
+	}
+
+	if (const UObjectMixerObjectFilter* Filter = GetMainObjectFilterInstance())
+	{
+		if (const TObjectPtr<UObject> Object = FObjectMixerUtils::GetRowObject(InTreeItem))
+		{
+			return Filter->GetRowDisplayName(Object, bIsHybridRow);
+		}
+	}
+
+	return FText::GetEmpty();
+}
+
+SObjectMixerEditorList* FObjectMixerEditorListRowData::GetListView() const
+{
+	return StaticCast<SObjectMixerEditorList*>(SceneOutlinerPtr);
+}
+
+EObjectMixerTreeViewMode FObjectMixerEditorListRowData::GetTreeViewMode()
+{
+	if (SObjectMixerEditorList* ListView = GetListView())
+	{
+		return ListView->GetTreeViewMode();
+	}
+
+	return EObjectMixerTreeViewMode::Folders;
+}
+
+TArray<TSharedPtr<ISceneOutlinerTreeItem>> FObjectMixerEditorListRowData::GetSelectedTreeViewItems() const
+{
+	if (SObjectMixerEditorList* ListView = GetListView())
+	{
+		return ListView->GetSelectedTreeViewItems();
+	}
+
+	return {};
+}
+
+void FObjectMixerEditorListRowData::OnChangeVisibility(const FSceneOutlinerTreeItemRef TreeItem, const bool bNewVisible)
+{
+	if (FObjectMixerEditorListRowData* RowData = FObjectMixerUtils::GetRowData(TreeItem))
+	{
+		if (SObjectMixerEditorList* ListView = RowData->GetListView())
+		{
+			ListView->ClearSoloRows();
+		}
+
+		RowData->SetUserHiddenInEditor(!bNewVisible);
+	}
+}
+
+bool FObjectMixerEditorListRowData::IsUserSetHiddenInEditor() const
+{
+	return VisibilityRules.bShouldBeHiddenInEditor;
+}
+
+void FObjectMixerEditorListRowData::SetUserHiddenInEditor(const bool bNewHidden)
+{
+	VisibilityRules.bShouldBeHiddenInEditor = bNewHidden;
+}
+
+bool FObjectMixerEditorListRowData::GetRowSoloState() const
+{
+	return VisibilityRules.bShouldBeSolo;
+}
+
+void FObjectMixerEditorListRowData::SetRowSoloState(const bool bNewSolo)
+{
+	VisibilityRules.bShouldBeSolo = bNewSolo;
+}
+
+void FObjectMixerEditorListRowData::ClearSoloRows() const
+{
+	if (SObjectMixerEditorList* ListView = GetListView())
+	{
+		ListView->ClearSoloRows();
+	}
+}
+
+bool FObjectMixerEditorListRowData::GetIsItemOrHybridChildSelected(const TSharedRef<ISceneOutlinerTreeItem> InRow)
+{
+	if (GetIsSelected(InRow))
+	{
+		return true;
+	}
+	
+	if (const TSharedPtr<ISceneOutlinerTreeItem> HybridChild = FObjectMixerUtils::GetHybridChild(InRow).Pin())
+	{
+		return FObjectMixerUtils::GetRowData(HybridChild)->GetIsSelected(HybridChild.ToSharedRef());
+	}
+	
+	return false;
+}
+
+void SetValueOnSelectedItems(
+	const FString& ValueAsString, const TArray<TSharedPtr<ISceneOutlinerTreeItem>>& OtherSelectedItems,
+	const FName& PropertyName, const TSharedPtr<ISceneOutlinerTreeItem> PinnedItem,
+	const EPropertyValueSetFlags::Type Flags)
+{
+	if (!ValueAsString.IsEmpty())
+	{
+		FScopedTransaction Transaction(
+			NSLOCTEXT("ObjectMixerEditor","OnPropertyChangedTransaction", "Object Mixer - Bulk Edit Selected Row Properties") );
+		
+		for (const TSharedPtr<ISceneOutlinerTreeItem>& SelectedRow : OtherSelectedItems)
+		{
+			const TWeakPtr<ISceneOutlinerTreeItem> SelectedHybridRow = FObjectMixerUtils::GetHybridChild(SelectedRow);
+			const TSharedPtr<ISceneOutlinerTreeItem> RowToUse = SelectedHybridRow.IsValid() ? SelectedHybridRow.Pin() : SelectedRow;
+
+			if (RowToUse == PinnedItem)
+			{
+				continue;
+			}
+
+			FObjectMixerEditorListRowData* RowData = FObjectMixerUtils::GetRowData(RowToUse);
+
+			// Skip folders
+			if (FObjectMixerUtils::AsFolderRow(RowToUse))
+			{
+				continue;
+			}
+
+			UObject* ObjectToModify = FObjectMixerUtils::GetRowObject(RowToUse);
+			
+			if (IsValid(ObjectToModify))
+			{
+				ObjectToModify->Modify();
+			}
+			else
+			{
+				UE_LOG(LogObjectMixerEditor, Warning, TEXT("%hs: Row '%s' has no valid associated object to modify."), __FUNCTION__, *RowData->GetDisplayName(RowToUse).ToString());
+				return;
+			}
+		
+			// Use handles if valid, otherwise use ImportText
+			if (const TWeakPtr<IPropertyHandle>* SelectedHandlePtr = RowData->PropertyNamesToHandles.Find(PropertyName))
+			{
+				if (SelectedHandlePtr->IsValid())
+				{				
+					SelectedHandlePtr->Pin()->SetValueFromFormattedString(ValueAsString, Flags);
+				}
+			}
+			else
+			{
+				if (const FProperty* PropertyToChange = FindFProperty<FProperty>(ObjectToModify->GetClass(), PropertyName))
+				{
+					if (void* ValuePtr = PropertyToChange->ContainerPtrToValuePtr<void>(ObjectToModify))
+					{
+						// Set the actual property value
+						PropertyToChange->ImportText_Direct(*ValueAsString, ValuePtr, ObjectToModify, PPF_None);
+					}
+				}
+			}
+		}
+	}
+}
+
+void FObjectMixerEditorListRowData::PropagateChangesToSimilarSelectedRowProperties(
+	const TSharedRef<ISceneOutlinerTreeItem> InRow, const FPropertyPropagationInfo PropertyPropagationInfo)
+{
+	if (PropertyPropagationInfo.PropertyName == NAME_None)
+	{
+		return;
+	}
+
+	if (const TSharedPtr<ISceneOutlinerTreeItem> RowToUse = FObjectMixerUtils::GetHybridChildOrRowItemIfNull(InRow); RowToUse)
+	{
+		if (!GetIsItemOrHybridChildSelected(InRow))
+		{
+			return;
+		}
+
+		FObjectMixerEditorListRowData* RowData = FObjectMixerUtils::GetRowData(RowToUse);
+		
+		const TWeakPtr<IPropertyHandle>* HandlePtr = RowData->PropertyNamesToHandles.Find(PropertyPropagationInfo.PropertyName);
+		if (HandlePtr->IsValid())
+		{
+			const TArray<TSharedPtr<ISceneOutlinerTreeItem>> OtherSelectedItems = RowData->GetSelectedTreeViewItems();
+			if (OtherSelectedItems.Num())
+			{
+				FString ValueAsString;
+				(*HandlePtr).Pin()->GetValueAsFormattedString(ValueAsString);
+			
+				SetValueOnSelectedItems(
+					ValueAsString, OtherSelectedItems, PropertyPropagationInfo.PropertyName,
+					RowToUse, PropertyPropagationInfo.PropertyValueSetFlags);
+			}
+		}
+	}
+}
+
+const FObjectMixerEditorListRowData::FTransientEditorVisibilityRules& FObjectMixerEditorListRowData::GetVisibilityRules() const
+{
+	return VisibilityRules;
+}
+
+void FObjectMixerEditorListRowData::SetVisibilityRules(const FTransientEditorVisibilityRules& InVisibilityRules)
+{
+	VisibilityRules = InVisibilityRules;
+}
diff --git a/Engine/Plugins/Editor/ObjectMixer/ObjectMixer/Source/ObjectMixer/Private/Views/List/ObjectMixerUtils.cpp b/Engine/Plugins/Editor/ObjectMixer/ObjectMixer/Source/ObjectMixer/Private/Views/List/ObjectMixerUtils.cpp
new file mode 100644
--- /dev/null
+++ b/Engine/Plugins/Editor/ObjectMixer/ObjectMixer/Source/ObjectMixer/Private/Views/List/ObjectMixerUtils.cpp
@@ -0,0 +1,252 @@
+// Copyright Epic Games, Inc. All Rights Reserved.
+
+#include "Views/List/ObjectMixerUtils.h"
+
+#include "ObjectMixerEditorSerializedData.h"
+#include "Views/List/SObjectMixerEditorList.h"
+
+#include "ScopedTransaction.h"
+#include "Views/List/RowTypes/ObjectMixerEditorListRowActor.h"
+#include "Views/List/RowTypes/ObjectMixerEditorListRowComponent.h"
+#include "Views/List/RowTypes/ObjectMixerEditorListRowFolder.h"
+#include "Views/List/RowTypes/ObjectMixerEditorListRowUObject.h"
+
+FObjectMixerEditorListRowFolder* FObjectMixerUtils::AsFolderRow(TSharedPtr<ISceneOutlinerTreeItem> InTreeItem)
+{
+	// Sometimes the outliner reports type as the parent type, but we can be assured we're always using our type.
+	// For this reason, we check for our type first and if it's not our type, check parent type then cast.
+	if (FObjectMixerEditorListRowFolder* AsFolder = InTreeItem->CastTo<FObjectMixerEditorListRowFolder>())
+	{
+		return AsFolder;
+	}
+	
+	if (FActorFolderTreeItem* AsFolder = InTreeItem->CastTo<FActorFolderTreeItem>())
+	{
+		if (AsFolder->GetFolder().IsValid())
+		{
+			return StaticCast<FObjectMixerEditorListRowFolder*>(AsFolder);
+		}
+	}
+
+	return nullptr;
+}
+
+FObjectMixerEditorListRowActor* FObjectMixerUtils::AsActorRow(TSharedPtr<ISceneOutlinerTreeItem> InTreeItem)
+{
+	// Sometimes the outliner reports type as the parent type, but we can be assured we're always using our type.
+	// For this reason, we check for our type first and if it's not our type, check parent type then cast.
+	if (FObjectMixerEditorListRowActor* AsActor = InTreeItem->CastTo<FObjectMixerEditorListRowActor>())
+	{
+		return AsActor;
+	}
+	
+	if (FActorTreeItem* AsActor = InTreeItem->CastTo<FActorTreeItem>())
+	{
+		if (AsActor->Actor.IsValid())
+		{
+			return StaticCast<FObjectMixerEditorListRowActor*>(AsActor);
+		}
+	}
+
+	return nullptr;
+}
+
+FObjectMixerEditorListRowComponent* FObjectMixerUtils::AsComponentRow(TSharedPtr<ISceneOutlinerTreeItem> InTreeItem)
+{
+	// Sometimes the outliner reports type as the parent type, but we can be assured we're always using our type.
+	// For this reason, we check for our type first and if it's not our type, check parent type then cast.
+	if (FObjectMixerEditorListRowComponent* AsComponent = InTreeItem->CastTo<FObjectMixerEditorListRowComponent>())
+	{
+		return AsComponent;
+	}
+	
+	if (FComponentTreeItem* AsComponent = InTreeItem->CastTo<FComponentTreeItem>())
+	{
+		if (AsComponent->Component.IsValid())
+		{
+			return StaticCast<FObjectMixerEditorListRowComponent*>(AsComponent);
+		}
+	}
+
+	return nullptr;
+}
+
+FObjectMixerEditorListRowUObject* FObjectMixerUtils::AsObjectRow(TSharedPtr<ISceneOutlinerTreeItem> InTreeItem)
+{
+	// Sometimes the outliner reports type as the parent type, but we can be assured we're always using our type.
+	// For this reason, we check for our type first and if it's not our type, check parent type then cast.
+	if (FObjectMixerEditorListRowUObject* AsObject = InTreeItem->CastTo<FObjectMixerEditorListRowUObject>())
+	{
+		return AsObject;
+	}
+	
+	if (FObjectMixerEditorListRowUObject* AsObject = InTreeItem->CastTo<FObjectMixerEditorListRowUObject>())
+	{
+		if (IsValid(AsObject->ObjectPtr))
+		{
+			return AsObject;
+		}
+	}
+
+	return nullptr;
+}
+
+FObjectMixerEditorListRowData* FObjectMixerUtils::GetRowData(TSharedPtr<ISceneOutlinerTreeItem> InTreeItem)
+{
+	if (FObjectMixerEditorListRowFolder* AsFolder = AsFolderRow(InTreeItem))
+	{
+		return &AsFolder->RowData;
+	}
+
+	if (FObjectMixerEditorListRowActor* AsActor = AsActorRow(InTreeItem))
+	{
+		return &AsActor->RowData;
+	}
+
+	if (FObjectMixerEditorListRowComponent* AsComponent = AsComponentRow(InTreeItem))
+	{
+		return &AsComponent->RowData;
+	}
+
+	if (FObjectMixerEditorListRowUObject* AsObject = AsObjectRow(InTreeItem))
+	{
+		return &AsObject->RowData;
+	}
+
+	return nullptr;
+}
+
+TWeakPtr<ISceneOutlinerTreeItem> FObjectMixerUtils::GetHybridChild(TSharedPtr<ISceneOutlinerTreeItem> InTreeItem)
+{
+	if (const FObjectMixerEditorListRowActor* ActorRow = AsActorRow(InTreeItem))
+	{
+		if (const TWeakObjectPtr<AActor> ThisObject = ActorRow->Actor; ThisObject.IsValid())
+		{
+			TArray<TWeakPtr<ISceneOutlinerTreeItem>> HybridCandidates;
+	
+			for (const TWeakPtr<ISceneOutlinerTreeItem>& ChildRow : InTreeItem->GetChildren())
+			{		
+				if (const FObjectMixerEditorListRowComponent* ComponentRow = AsComponentRow(ChildRow.Pin()))
+				{
+					if (TWeakObjectPtr<UActorComponent> ChildObject = ComponentRow->Component; ChildObject.IsValid())
+					{
+						if (ChildObject->GetOuter() == ThisObject)
+						{
+							HybridCandidates.Add(ChildRow);
+
+							if (HybridCandidates.Num() > 1)
+							{
+								// There can only be one row to hybrid with.
+								// If there's more than one candidate, don't hybrid.
+								break;
+							}
+						}
+					}
+				}
+			}
+		
+			if (HybridCandidates.Num() == 1)
+			{
+				return HybridCandidates[0];
+			}
+		}
+	}
+	
+	return nullptr;
+}
+
+TSharedRef<ISceneOutlinerTreeItem> FObjectMixerUtils::GetHybridChildOrRowItemIfNull(const TSharedRef<ISceneOutlinerTreeItem> InRow)
+{
+	if (const TWeakPtr<ISceneOutlinerTreeItem> HybridChild = GetHybridChild(InRow); HybridChild.IsValid())
+	{
+		return HybridChild.Pin().ToSharedRef();
+	}
+		
+	return InRow;
+}
+
+UObject* FObjectMixerUtils::GetRowObject(TSharedPtr<ISceneOutlinerTreeItem> InTreeItem)
+{
+	if (const FObjectMixerEditorListRowActor* ActorRow = AsActorRow(InTreeItem))
+	{
+		return ActorRow->Actor.Get();
+	}
+
+	if (const FObjectMixerEditorListRowComponent* ComponentRow = AsComponentRow(InTreeItem))
+	{
+		return ComponentRow->Component.Get();
+	}
+
+	if (const FObjectMixerEditorListRowUObject* ObjectRow = AsObjectRow(InTreeItem))
+	{
+		return ObjectRow->ObjectPtr.Get();
+	}
+
+	return nullptr;
+}
+
+AActor* FObjectMixerUtils::GetSelfOrOuterAsActor(TSharedPtr<ISceneOutlinerTreeItem> InTreeItem)
+{
+	if (UObject* Object = GetRowObject(InTreeItem))
+	{
+		AActor* Actor = Cast<AActor>(Object);
+
+		if (!Actor)
+		{
+			Actor = Object->GetTypedOuter<AActor>();
+		}
+
+		return Actor;
+	}
+
+	return nullptr;
+}
+
+bool FObjectMixerUtils::IsObjectRefInCollection(const FName& CollectionName, const UObject* Object, const TSharedPtr<FObjectMixerEditorList> ListModel)
+{	
+	if (Object)
+	{
+		if (CollectionName == UObjectMixerEditorSerializedData::AllCollectionName)
+		{
+			return true;
+		}
+
+		return ListModel->IsObjectInCollection(CollectionName, Object);
+	}
+	
+	return false;
+}
+
+bool FObjectMixerUtils::IsObjectRefInCollection(const FName& CollectionName, TSharedPtr<ISceneOutlinerTreeItem> InTreeItem)
+{	
+	if (const UObject* Object = GetRowObject(InTreeItem))
+	{
+		if (SObjectMixerEditorList* ListView = GetRowData(InTreeItem)->GetListView())
+		{
+			if (const TSharedPtr<FObjectMixerEditorList> ListModel = ListView->GetListModelPtr().Pin())
+			{
+				return IsObjectRefInCollection(CollectionName, Object, ListModel);
+			}
+		}
+	}
+	
+	return false;
+}
+
+void FObjectMixerUtils::SetChildRowsSelected(
+	TSharedPtr<ISceneOutlinerTreeItem> InTreeItem, const bool bNewSelected, const bool bRecursive)
+{
+	for (const TWeakPtr<ISceneOutlinerTreeItem>& ChildRow : InTreeItem->GetChildren())
+	{
+		if (const TSharedPtr<ISceneOutlinerTreeItem> PinnedChildRow = ChildRow.Pin())
+		{
+			// Recurse even if not visible
+			if (bRecursive)
+			{
+				SetChildRowsSelected(PinnedChildRow, bNewSelected, bRecursive);
+			}
+	
+			GetRowData(PinnedChildRow)->SetIsSelected(PinnedChildRow.ToSharedRef(), bNewSelected);
+		}
+	}
+}
