Unshelved from pending changelist '23770228':

diff --git a/Engine/Plugins/Editor/ObjectMixer/ObjectMixer/Source/ObjectMixer/Private/Views/List/SObjectMixerEditorList.cpp b/Engine/Plugins/Editor/ObjectMixer/ObjectMixer/Source/ObjectMixer/Private/Views/List/SObjectMixerEditorList.cpp
--- a/Engine/Plugins/Editor/ObjectMixer/ObjectMixer/Source/ObjectMixer/Private/Views/List/SObjectMixerEditorList.cpp
+++ b/Engine/Plugins/Editor/ObjectMixer/ObjectMixer/Source/ObjectMixer/Private/Views/List/SObjectMixerEditorList.cpp
@@ -2,35 +2,48 @@
 
 #include "Views/List/SObjectMixerEditorList.h"
 
-// View Filters
-#include "Algo/Find.h"
-#include "Views/List/ObjectMixerEditorListFilters/IObjectMixerEditorListFilter.h"
-
-#include "Engine/Level.h"
-#include "ObjectMixerEditorLog.h"
-#include "Framework/Application/SlateApplication.h"
+#include "ObjectMixerEditorModule.h"
+#include "ObjectMixerEditorSerializedData.h"
 #include "ObjectMixerEditorSettings.h"
-#include "Framework/Views/TableViewMetadata.h"
-#include "ObjectMixerEditorStyle.h"
+#include "ObjectFilter/ObjectMixerEditorObjectFilter.h"
 #include "Views/List/ObjectMixerEditorList.h"
-#include "Views/List/SObjectMixerEditorListRow.h"
+#include "Views/List/ObjectMixerUtils.h"
+#include "Views/List/Modes/ObjectMixerOutlinerMode.h"
+#include "Views/List/Modes/OutlinerColumns/ObjectMixerOutlinerColumns.h"
+#include "Views/List/RowTypes/ObjectMixerEditorListRowActor.h"
+#include "Views/Widgets/SCollectionSelectionButton.h"
+#include "Views/Widgets/SObjectMixerPlacementAssetMenuEntry.h"
 
-#include "Algo/AllOf.h"
-#include "Algo/AnyOf.h"
+// View Filters
+#include "ActorDescTreeItem.h"
+#include "ActorFolderTreeItem.h"
+#include "ComponentTreeItem.h"
+#include "Editor.h"
 #include "EditorActorFolders.h"
+#include "IPlacementModeModule.h"
+#include "LevelEditor.h"
+#include "LevelEditorContextMenu.h"
+#include "SceneOutlinerModule.h"
+#include "SceneOutlinerTextInfoColumn.h"
+#include "SPositiveActionButton.h"
+#include "Algo/Find.h"
+#include "Engine/Level.h"
+#include "Framework/Application/SlateApplication.h"
 #include "Framework/MultiBox/MultiBoxBuilder.h"
-#include "LevelEditorViewport.h"
-#include "ObjectMixerEditorModule.h"
-#include "ObjectMixerEditorSerializedData.h"
-#include "ScopedTransaction.h"
-#include "Selection.h"
-#include "Views/MainPanel/ObjectMixerEditorMainPanel.h"
-#include "Views/Widgets/ObjectMixerEditorListMenuContext.h"
-#include "Widgets/Layout/SWidgetSwitcher.h"
-#include "Widgets/Text/SRichTextBlock.h"
+#include "Framework/Views/TableViewMetadata.h"
+#include "LevelInstance/LevelInstanceInterface.h"
+#include "LevelInstance/LevelInstanceSubsystem.h"
+#include "ToolMenus.h"
+#include "Widgets/Input/SCheckBox.h"
+#include "Widgets/Layout/SWrapBox.h"
 #include "Widgets/Views/STreeView.h"
+#include "WorldPartition/ActorDescContainer.h"
+#include "WorldPartition/WorldPartitionActorDesc.h"
+#include "WorldPartition/ContentBundle/ContentBundleDescriptor.h"
+#include "WorldPartition/ContentBundle/ContentBundleEngineSubsystem.h"
+#include "WorldPartition/DataLayer/DataLayerSubsystem.h"
 
-#define LOCTEXT_NAMESPACE "ObjectMixerEditorList"
+#define LOCTEXT_NAMESPACE "SObjectMixerEditorList"
 
 const FName SObjectMixerEditorList::ItemNameColumnName(TEXT("Builtin_Name"));
 const FName SObjectMixerEditorList::EditorVisibilityColumnName(TEXT("Builtin_EditorVisibility"));
@@ -39,426 +52,431 @@ const FName SObjectMixerEditorList::EditorVisibilitySoloColumnName(TEXT("Builtin
 void SObjectMixerEditorList::Construct(const FArguments& InArgs, TSharedRef<FObjectMixerEditorList> ListModel)
 {
 	ListModelPtr = ListModel;
+
+	FSceneOutlinerInitializationOptions InitOptions;
+	InitOptions.bShowSearchBox = true;
+	InitOptions.bFocusSearchBoxWhenOpened = true;
+	InitOptions.bShowCreateNewFolder = true;
+	InitOptions.bShowParentTree = true;
+	InitOptions.OutlinerIdentifier = ListModel->GetModuleName();
 	
-	// Set Default Sorting info
-	ActiveSortingType = EColumnSortMode::Ascending;
-	
-	HeaderRow = SNew(SHeaderRow)
-				.CanSelectGeneratedColumn(false)
-				.Visibility(EVisibility::Visible)
-				;
-	
-	ChildSlot
-	[
-		SNew(SVerticalBox)
+	// Context Menu
+	{
+		UToolMenus* ToolMenus = UToolMenus::Get();
+		static const FName MenuName = "LevelEditor.LevelEditorSceneOutliner.ContextMenu";
+		if (!ToolMenus->IsMenuRegistered(MenuName))
+		{
+			UToolMenu* Menu = ToolMenus->RegisterMenu(MenuName, "SceneOutliner.DefaultContextMenuBase");
+			FToolMenuSection& Section = Menu->AddDynamicSection("LevelEditorContextMenu",
+				FNewToolMenuDelegate::CreateRaw(this, &SObjectMixerEditorList::OnAddDynamicContextMenuSection)
+			);
+			Section.InsertPosition = FToolMenuInsert("MainSection", EToolMenuInsertType::Before);
+		}
 
-		+ SVerticalBox::Slot()
-		.HAlign(HAlign_Fill)
-		.VAlign(VAlign_Fill)
-		[
-			SNew(SWidgetSwitcher)
-			.WidgetIndex_Lambda([this]()
+		InitOptions.ModifyContextMenu.BindLambda([=](FName& OutMenuName, FToolMenuContext& MenuContext)
 			{
-				return DoesTreeViewHaveVisibleChildren() ? 0 : 1;
-			})
-			
-			+ SWidgetSwitcher::Slot()
-			.HAlign(HAlign_Fill)
-			.VAlign(VAlign_Fill)
-			.Padding(2.0f, 2.0f, 2.0f, 2.0f)
-			[
-				SAssignNew(TreeViewPtr, STreeView<FObjectMixerEditorListRowPtr>)
-				.HeaderRow(HeaderRow)
-				.SelectionMode(ESelectionMode::Multi)
-				.HighlightParentNodesForSelection(true)
-				.OnContextMenuOpening_Lambda([this]()
-				{
-					TSharedPtr<SWidget> Widget;
-					
-					if (GetListModelPtr().IsValid())
-					{
-						SelectedTreeItemsToSelectedInLevelEditor();
-						Widget = UObjectMixerEditorListMenuContext::CreateContextMenu(
-							{GetSelectedTreeViewItems(), GetListModelPtr().Pin()->GetMainPanelModel()});
-					}
-					
-					return Widget;
-				})
-				.OnSelectionChanged_Lambda([this] (const FObjectMixerEditorListRowPtr& Row, const ESelectInfo::Type SelectionType)
-				{
-					if (SelectionType != ESelectInfo::Direct) // Don't call if selected from code
-					{
-						const bool bSyncSelectionEnabled = GetDefault<UObjectMixerEditorSettings>()->bSyncSelection;
-						const bool bIsAltDown = FSlateApplication::Get().GetModifierKeys().IsAltDown();
-						const bool bShouldSyncSelection = bSyncSelectionEnabled ? !bIsAltDown : bIsAltDown;
-						
-						if (bShouldSyncSelection)
-						{
-							SelectedTreeItemsToSelectedInLevelEditor();
-						}
-					}
-				})
-				.OnMouseButtonDoubleClick_Lambda([this] (FObjectMixerEditorListRowPtr Row)
-				{							
-					if (Row->GetRowType() == FObjectMixerEditorListRow::MatchingObject ||
-						Row->GetRowType() == FObjectMixerEditorListRow::ContainerObject)
-					{
-						if (GCurrentLevelEditingViewportClient)
-						{
-							if (const AActor* Actor = Row->GetSelfOrOuterAsActor())
-							{
-								FVector Origin;
-								FVector Extents;
-								Actor->GetActorBounds(false, Origin, Extents, true);
-								GCurrentLevelEditingViewportClient->FocusViewportOnBox(FBox(Origin - Extents, Origin + Extents));
-							}
-						}
-					}
-					else
-					{
-						Row->SetIsTreeViewItemExpanded(!Row->GetIsTreeViewItemExpanded());
-					}
-				})
-				.TreeItemsSource(&VisibleTreeViewObjects)
-				.OnGenerateRow_Lambda([this](FObjectMixerEditorListRowPtr Row, const TSharedRef<STableViewBase>& OwnerTable)
-					{
-						check(Row.IsValid());
-					
-						return SNew(SObjectMixerEditorListRow, TreeViewPtr.ToSharedRef(), Row)
-								.Visibility_Raw(Row.Get(), &FObjectMixerEditorListRow::GetDesiredRowWidgetVisibility);
-					})
-				.OnGetChildren_Raw(this, &SObjectMixerEditorList::OnGetRowChildren)
-				.OnExpansionChanged_Raw(this, &SObjectMixerEditorList::OnRowChildExpansionChange, false)
-				.OnSetExpansionRecursive(this, &SObjectMixerEditorList::OnRowChildExpansionChange, true)
-			]
-
-			// For when no rows exist in view
-			+ SWidgetSwitcher::Slot()
-			.HAlign(HAlign_Fill)
-			.Padding(2.0f, 24.0f, 2.0f, 2.0f)
-			[
-				SNew(SRichTextBlock)
-				.DecoratorStyleSet(&FAppStyle::Get())
-				.AutoWrapText(true)
-				.Justification(ETextJustify::Center)
-				.Text_Lambda([this]()
+				OutMenuName = MenuName;
+
+				const FLevelEditorModule& LevelEditorModule = FModuleManager::LoadModuleChecked<FLevelEditorModule>("LevelEditor");
+				const TSharedPtr<ILevelEditor> LevelEditorPtr = LevelEditorModule.GetLevelEditorInstance().Pin();
+
+				if (LevelEditorPtr.IsValid())
 				{
-					// Preset Empty List (with filter)
-                    return LOCTEXT("EmptyListPresetWithFilter", "No matching items in your list.\n\nCheck your filters.");
-				})
-			]
-		]
-	];
-}
+					FLevelEditorContextMenu::InitMenuContext(MenuContext, LevelEditorPtr, ELevelEditorMenuContext::SceneOutliner);
+				}
+			});
+	}
+	
+	// Mode
+	InitOptions.ModeFactory = FCreateSceneOutlinerMode::CreateLambda([this](SSceneOutliner* InOutliner)
+	{		
+		return new FObjectMixerOutlinerMode(
+			FObjectMixerOutlinerModeParams(
+				InOutliner, nullptr, false,
+				false, true, true
+			),
+			ListModelPtr.Pin().ToSharedRef()
+		);
+	});
 
-SObjectMixerEditorList::~SObjectMixerEditorList()
-{	
-	HeaderRow.Reset();
+	// Header Row
+	InitOptions.bShowHeaderRow = true;
+	InitOptions.bCanSelectGeneratedColumns = false;
+	SetupColumns(InitOptions);
+	 
+	SSceneOutliner::Construct(InArgs, InitOptions);
+
+	InsertCollectionSelector();
 
-	FlushMemory(false);
+	if (UObjectMixerEditorSerializedData* SerializedData = GetMutableDefault<UObjectMixerEditorSerializedData>())
+	{
+		SerializedData->OnObjectMixerCollectionMapChanged.AddRaw(this, &SObjectMixerEditorList::RebuildCollectionSelector);
+	}
+	
+	RebuildCollectionSelector();
 
-	TreeViewPtr.Reset();
+	SetSingleCollectionSelection();
 }
 
-UWorld* SObjectMixerEditorList::GetWorld() const
+
+FReply SObjectMixerEditorList::OnKeyDown(const FGeometry& MyGeometry, const FKeyEvent& InKeyEvent)
 {
-	return FObjectMixerEditorModule::Get().GetWorld();
+	// Check to see if any actions can be processed
+	// If we are in debug mode do not process commands
+	if (FSlateApplication::Get().IsNormalExecution())
+	{
+		if (GetListModelPtr().Pin()->ObjectMixerElementEditCommands->ProcessCommandBindings(InKeyEvent))
+		{
+			return FReply::Handled();
+		}
+	}
+	
+	return SCompoundWidget::OnKeyDown(MyGeometry, InKeyEvent);
 }
 
-void SObjectMixerEditorList::RebuildList()
+bool SObjectMixerEditorList::CanCreateFolder() const
 {
-	bIsRebuildRequested = false;
+	return GetSelectedTreeViewItemCount() > 0;
+}
+
+void SObjectMixerEditorList::OnAddDynamicContextMenuSection(UToolMenu* InMenu)
+{
+	FName LevelContextMenuName = "LevelEditor.ActorContextMenu";
+
+	const FSceneOutlinerItemSelection Selection = GetSelection();
 	
-	GenerateTreeView();
+	if (Selection.Has<FComponentTreeItem>())
+	{
+		LevelContextMenuName = "LevelEditor.ComponentContextMenu";
+	}
+	else if (Selection.Has<FActorTreeItem>())
+	{
+		LevelContextMenuName = "LevelEditor.ActorContextMenu";
+	}
+	else if (Selection.Num() > 0)
+	{
+		LevelContextMenuName = "LevelEditor.ElementContextMenu";
+	}
+				
+	if (LevelContextMenuName != NAME_None)
+	{
+		// Extend the menu even if no actors selected, as Edit menu should always exist for scene outliner
+		UToolMenu* OtherMenu = UToolMenus::Get()->GenerateMenu(LevelContextMenuName, InMenu->Context);
+		InMenu->Sections.Append(OtherMenu->Sections);
+	}
 }
 
-void SObjectMixerEditorList::RefreshList()
+TSharedRef<SWidget> SObjectMixerEditorList::OnGenerateAddObjectButtonMenu() const
 {
-	if (TreeViewRootObjects.Num() > 0)
+	TSet<UClass*> ClassesToPlace;
+	for (const TObjectPtr<UObjectMixerObjectFilter>& Instance : ListModelPtr.Pin()->GetObjectFilterInstances())
+	{
+		const TSet<TSubclassOf<AActor>> SubclassesOfActor = Instance->GetObjectClassesToPlace();
+		if (SubclassesOfActor.Num() > 0)
+		{
+			ClassesToPlace.Append(
+				Instance->GetParentAndChildClassesFromSpecifiedClasses(
+					SubclassesOfActor,
+					Instance->GetObjectMixerPlacementClassInclusionOptions()
+				)
+			);
+		}
+	}
+
+	if (ClassesToPlace.Num() > 0)
 	{
-		// Apply last search
-		ExecuteListViewSearchOnAllRows(GetSearchStringFromSearchInputField(), false);
+		FMenuBuilder AddObjectButtonMenuBuilder = FMenuBuilder(true, nullptr);
 
-		// Enforce Sort
-		const FName& SortingName = GetActiveSortingColumnName();
-		ExecuteSort(SortingName, GetSortModeForColumn(SortingName), false);
+		for (const UClass* Class : ClassesToPlace)
+		{
+			if (const UActorFactory* Factory = GEditor->FindActorFactoryForActorClass(Class))
+			{
+				AddObjectButtonMenuBuilder.AddWidget(
+					SNew(SObjectMixerPlacementAssetMenuEntry, MakeShareable(new FPlaceableItem(*Factory->GetClass()))), FText::GetEmpty());
+			}
+		}
 
-		// Show/Hide rows based on SetBy changes and filter settings
-		EvaluateIfRowsPassFilters(false);
+		return AddObjectButtonMenuBuilder.MakeWidget();
 	}
 
-	FindVisibleObjectsAndRequestTreeRefresh();
+	return
+		SNew(SBox)
+		.Padding(5)
+		[
+			SNew(STextBlock)
+				.Text(LOCTEXT("NoPlaceableActorsDefinedWarning", "Please define some placeable actors in the\nfilter class by overriding ForceGetObjectClassesToPlace."))
+				.Font(FAppStyle::Get().GetFontStyle("NormalFontItalic"))
+		]
+	;
 }
 
-void SObjectMixerEditorList::RequestRebuildList(const FString& InItemToScrollTo)
+void SObjectMixerEditorList::SetSingleCollectionSelection(const FName& CollectionToEnableName)
 {
-	bIsRebuildRequested = true;
+	// Disable all collection filters except CollectionToEnableName
+	GetListModelPtr().Pin()->SetSelectedCollections({CollectionToEnableName});
 }
 
-void SObjectMixerEditorList::OnRenameCommand()
-{
-	if (GetSelectedTreeViewItemCount() == 1)
-	{
-		const FObjectMixerEditorListRowPtr TreeViewItem = GetSelectedTreeViewItems()[0];
 
-		PendingRenameItem = TreeViewItem->GetUniqueIdentifier();
-	}
+EObjectMixerTreeViewMode SObjectMixerEditorList::GetTreeViewMode()
+{
+	return ListModelPtr.Pin()->GetTreeViewMode();
 }
 
-void SObjectMixerEditorList::AddToPendingPropertyPropagations(
-	const FObjectMixerEditorListRow::FPropertyPropagationInfo& InPropagationInfo)
+void SObjectMixerEditorList::SetTreeViewMode(EObjectMixerTreeViewMode InViewMode)
 {
-	PendingPropertyPropagations.Add(InPropagationInfo);
+	ListModelPtr.Pin()->SetTreeViewMode(InViewMode);
 }
 
-TArray<FObjectMixerEditorListRowPtr> SObjectMixerEditorList::GetSelectedTreeViewItems() const
+const TSet<FName>& SObjectMixerEditorList::GetSelectedCollections()
 {
-	return TreeViewPtr->GetSelectedItems();
+	return GetListModelPtr().Pin()->GetSelectedCollections();
 }
 
-int32 SObjectMixerEditorList::GetSelectedTreeViewItemCount() const
+bool SObjectMixerEditorList::IsCollectionSelected(const FName& CollectionName)
 {
-	return TreeViewPtr->GetSelectedItems().Num();
+	return GetListModelPtr().Pin()->IsCollectionSelected(CollectionName);
 }
 
-void SObjectMixerEditorList::SetTreeViewItemSelected(FObjectMixerEditorListRowPtr Item, const bool bNewSelected)
+void SObjectMixerEditorList::RebuildCollectionSelector()
 {
-	TreeViewPtr->SetItemSelection(Item, bNewSelected);
+	// Make user collections
+
+	CollectionSelectorBox->ClearChildren();
+	CollectionSelectorBox->SetVisibility(EVisibility::Collapsed);
+
+	auto CreateCollectionFilterAndAddToCollectionSelector =
+		[this](const FName& InCollectionName) 
+	{
+		CollectionSelectorBox->AddSlot()
+		[
+			SNew(SCollectionSelectionButton, SharedThis(this), InCollectionName)
+		];
+	};
+
+	TArray<FName> AllCollectionNames = ListModelPtr.Pin()->GetAllCollectionNames();
+	
+	// If no collections, rather than show "All" just keep the box hidden
+	// Otherwise make "All" collection widget
+	if (AllCollectionNames.Num())
+	{
+		CreateCollectionFilterAndAddToCollectionSelector(UObjectMixerEditorSerializedData::AllCollectionName);
+	}
+	
+	for (const FName& CollectionName : AllCollectionNames)
+	{
+		// Create widgets for each name
+		CreateCollectionFilterAndAddToCollectionSelector(CollectionName);
+	}
+
+	CollectionSelectorBox->SetVisibility(EVisibility::Visible);
 }
 
-void SObjectMixerEditorList::SyncEditorSelectionToListSelection()
+bool SObjectMixerEditorList::RequestRemoveCollection(const FName& CollectionName)
 {
-	if (GetDefault<UObjectMixerEditorSettings>()->bSyncSelection && !bShouldPauseSyncSelection)
+	if (ListModelPtr.Pin()->RequestRemoveCollection(CollectionName))
 	{
-		bIsEditorToListSelectionSyncRequested = false;
-		const USelection* SelectedActors = GEditor->GetSelectedActors();
+		OnCollectionCheckedStateChanged(true, UObjectMixerEditorSerializedData::AllCollectionName);
 
-		if (const int32 SelectionCount = SelectedActors->Num(); SelectionCount > 0)
-		{
-			bShouldPauseSyncSelection = true;
-			TreeViewPtr->ClearSelection();
+		RebuildCollectionSelector();
 
-			FObjectMixerEditorListRowPtr* ListRowPtr = nullptr;
-			for (int32 SelectionItr = 0; SelectionItr < SelectionCount; SelectionItr++)
-			{
-				UObject* SelectedActor = SelectedActors->GetSelectedObject(SelectionItr);
+		return true;
+	}
 
-				ListRowPtr = ObjectsToRowsCreated.Find(SelectedActor);
+	return false;
+}
 
-				if (ListRowPtr)
-				{
-					TreeViewPtr->SetItemSelection(*ListRowPtr, true);
-				}
-			}
-			
-			// Scroll to the last selected item
-			if (ListRowPtr)
-			{
-				TreeViewPtr->RequestScrollIntoView(*ListRowPtr);
-			}
-		}
-	}
+bool SObjectMixerEditorList::RequestDuplicateCollection(const FName& CollectionToDuplicateName, FName& DesiredDuplicateName) const
+{
+	return ListModelPtr.Pin()->RequestDuplicateCollection(CollectionToDuplicateName, DesiredDuplicateName);
+}
+
+bool SObjectMixerEditorList::RequestRenameCollection(
+	const FName& CollectionNameToRename,
+	const FName& NewCollectionName)
+{
+	return ListModelPtr.Pin()->RequestRenameCollection(CollectionNameToRename, NewCollectionName);
 }
 
-void SObjectMixerEditorList::OnRequestNewFolder(TOptional<FFolder> ExplicitParentFolder)
+bool SObjectMixerEditorList::DoesCollectionExist(const FName& CollectionName) const
 {
-	const FScopedTransaction Transaction(LOCTEXT("UndoAction_CreateFolder", "Create Folder"));
+	return ListModelPtr.Pin()->DoesCollectionExist(CollectionName);
+}
+
+void SObjectMixerEditorList::OnCollectionCheckedStateChanged(bool bShouldBeChecked, FName CollectionName)
+{
+	GetListModelPtr().Pin()->OnPreFilterChangeDelegate.Broadcast();
+
+	const bool bIsAllCollection = CollectionName.IsEqual(UObjectMixerEditorSerializedData::AllCollectionName);
 
-	UWorld* World = GetWorld();
-	FFolder NewFolderPath;
-	const TArray<FObjectMixerEditorListRowPtr>& SelectedTreeItems = TreeViewPtr->GetSelectedItems();
-	if (ExplicitParentFolder.IsSet())
+	if (FSlateApplication::Get().GetModifierKeys().IsControlDown())
 	{
-		NewFolderPath =
-			FActorFolders::Get().GetDefaultFolderName(*World, ExplicitParentFolder.GetValue());
+		if (bIsAllCollection)
+		{
+			// Can't disable All
+			return;
+		}
+		
+		GetListModelPtr().Pin()->SetCollectionSelected(CollectionName, bShouldBeChecked);
+
+		if (!bShouldBeChecked && GetSelectedCollections().Num() == 0)
+		{
+			// Reset to all
+			SetSingleCollectionSelection();
+		}
 	}
 	else
 	{
-		// Explicit folder not provided - try to find an appropriate parent folder based on selection
-		if (SelectedTreeItems.Num())
+		if (!bShouldBeChecked && bIsAllCollection)
 		{
-			FFolder FoundParent;
-			int32 FolderCount = 0;
-
-			// Check for a folder row first. It's possible for a root folder to be selected.
-			// If that's the case, it won't have a parent folder but we won't want to move it to a selected actor's direct folder.
-			// For this reason, we only check selected actors if there are no folders selected or there is no folder match
-			const TSharedPtr<FObjectMixerEditorListRow>* FolderMatch = Algo::FindByPredicate(
-				SelectedTreeItems,
-				[&FoundParent, &FolderCount](const FObjectMixerEditorListRowPtr& Other)
-				{
-					if (Other->GetRowType() == FObjectMixerEditorListRow::Folder)
-					{
-						FolderCount++;
-						FoundParent = Other->GetFolder().GetParent();
-					}
+			// Reset to all
+			SetSingleCollectionSelection();
+		}
+		
+		// Set just this filter active
+		SetSingleCollectionSelection(CollectionName);
+	}
 
-					return !FoundParent.IsNone();
-				}
-			);
+	GetListModelPtr().Pin()->OnPostFilterChangeDelegate.Broadcast();
+}
 
-			// If no folder row found, try actors
-			if (FolderCount == 0 || !FolderMatch)
-			{
-				const TSharedPtr<FObjectMixerEditorListRow>* ActorMatch = Algo::FindByPredicate(
-					SelectedTreeItems,
-					[&FoundParent](const FObjectMixerEditorListRowPtr& Other)
-					{
-						if (const AActor* Actor = Other->GetSelfOrOuterAsActor())
-						{
-							FoundParent = Actor->GetFolder();
-						}
+ECheckBoxState SObjectMixerEditorList::GetCollectionCheckedState(FName CollectionName) const
+{
+	return GetListModelPtr().Pin()->IsCollectionSelected(CollectionName) ? ECheckBoxState::Checked : ECheckBoxState::Unchecked;
+}
 
-						return !FoundParent.IsNone();
-					}
-				);
-			}
+FObjectMixerOutlinerMode* SObjectMixerEditorList::GetCastedMode() const
+{
+	return StaticCast<FObjectMixerOutlinerMode*>(Mode);
+}
 
-			if (!FoundParent.IsNone())
-			{
-				NewFolderPath =
-				   FActorFolders::Get().GetDefaultFolderName(*World, FoundParent);
-			}
+UWorld* SObjectMixerEditorList::GetWorld() const
+{
+	if (FObjectMixerOutlinerMode* CastedMode = GetCastedMode())
+	{
+		if (UWorld* UserChosenWorld = CastedMode->GetRepresentingWorld())
+		{
+			return UserChosenWorld;
 		}
 	}
+	
+	return FObjectMixerEditorModule::Get().GetWorld();
+}
 
-	// Fallback
-	if (NewFolderPath.IsNone())
-	{
-		NewFolderPath = FActorFolders::Get().GetDefaultFolderForSelection(*World);
-	}
+void SObjectMixerEditorList::RebuildList()
+{
+	bIsRebuildRequested = false;
+	
+	Mode->Rebuild();
+	RebuildCollectionSelector();
+}
+
+void SObjectMixerEditorList::RefreshList()
+{
+	GetTreeView()->RequestListRefresh();
+}
+
+void SObjectMixerEditorList::RequestRebuildList(const FString& InItemToScrollTo)
+{
+	bIsRebuildRequested = true;
+}
 
-	if (FActorFolders::Get().CreateFolder(*World, NewFolderPath))
+void SObjectMixerEditorList::InsertCollectionSelector()
+{
+	if (const TSharedPtr<SHorizontalBox> PinnedToolbar = ToolbarPtr.Pin())
 	{
-		if (SelectedTreeItems.Num())
+		if (TSharedPtr<SVerticalBox> VerticalBoxParent = StaticCastSharedPtr<SVerticalBox>(PinnedToolbar->GetParentWidget()))
 		{
-			TArray<AActor*> ActorsToMove;
-			TSet<FFolder> FoldersToMove;
-			for (const FObjectMixerEditorListRowPtr& SelectedRow : SelectedTreeItems)
-			{
-				if (SelectedRow->GetRowType() == FObjectMixerEditorListRow::MatchingObject ||
-					SelectedRow->GetRowType() == FObjectMixerEditorListRow::ContainerObject)
-				{
-					FObjectMixerEditorListRowPtr PinnedParent = SelectedRow->GetDirectParentRow().Pin();
-					if (PinnedParent.IsValid() &&
-						PinnedParent->GetRowType() == FObjectMixerEditorListRow::Folder &&
-						SelectedTreeItems.Contains(PinnedParent))
-					{
-						// If the selected item's parent is a folder and it's selected, just move the folder.
-						// It will take the selected objects along with it.
-						const FFolder& ParentFolder = SelectedRow->GetFolder();
-
-						if (!ParentFolder.IsNone())
-						{
-							FoldersToMove.Add(ParentFolder);
-						}
-					}
-					else if (AActor* Actor = SelectedRow->GetSelfOrOuterAsActor())
-					{
-						ActorsToMove.Add(Actor);
-					}
-				}
-				else if (SelectedRow->GetRowType() == FObjectMixerEditorListRow::Folder)
-				{
-					FoldersToMove.Add(SelectedRow->GetFolder());
-				}
-			}
+			int32 InsertAt = 2; // Best guess in case we can't find it procedurally
 
-			if (FoldersToMove.Num())
+			for (int32 SlotItr = 0; SlotItr < VerticalBoxParent->NumSlots(); SlotItr++)
 			{
-				for (const FFolder& Folder : FoldersToMove)
-				{
-					OnRequestMoveFolder(Folder, NewFolderPath);
-				}
-			}
-
-			if (ActorsToMove.Num())
-			{				
-				for (AActor* Actor : ActorsToMove)
+				if (VerticalBoxParent->IsValidSlotIndex(SlotItr))
 				{
-					Actor->Modify();
-					Actor->SetFolderPath_Recursively(NewFolderPath.GetPath());
-
-					if (AActor* AttachParent = Actor->GetAttachParentActor())
+					if (VerticalBoxParent->GetSlot(SlotItr).GetWidget() == PinnedToolbar)
 					{
-						if (!ActorsToMove.Contains(AttachParent))
-						{
-							AttachParent->Modify();
-							FDetachmentTransformRules DetachmentRules(EDetachmentRule::KeepWorld, false);
-							Actor->DetachFromActor(DetachmentRules);
-						}
+						InsertAt = SlotItr + 1;
 					}
 				}
 			}
+			
+			VerticalBoxParent->InsertSlot(InsertAt)
+			.Padding(8, 2, 8, 7)
+			.AutoHeight()
+			[
+				SAssignNew(CollectionSelectorBox, SWrapBox)
+				.UseAllottedSize(true)
+				.InnerSlotPadding(FVector2D(4,4))
+			];
 		}
-
-		// Set new folder row to be renamed on next frame
-		PendingRenameItem.RowName = NewFolderPath.GetPath().ToString();
 	}
 }
 
-void SObjectMixerEditorList::OnRequestMoveFolder(const FFolder& FolderToMove, const FFolder& TargetNewParentFolder)
+void SObjectMixerEditorList::OnRenameCommand()
 {
-	if (UWorld* EditorWorld = GetWorld())
-	{
-		if (FolderToMove.GetParent() != TargetNewParentFolder)
-		{
-			check(TargetNewParentFolder.GetRootObject() == FolderToMove.GetRootObject());
-			// Get unique name
-			const FFolder NewPath = FActorFolders::Get().GetFolderName(*EditorWorld, TargetNewParentFolder, FolderToMove.GetLeafName());
-			FActorFolders::Get().RenameFolderInWorld(*EditorWorld, FolderToMove, NewPath);
-		}
-	}
+	
 }
 
-void SObjectMixerEditorList::SetSelectedTreeViewItemActorsEditorVisible(const bool bNewIsVisible, const bool bIsRecursive)
+void SObjectMixerEditorList::AddToPendingPropertyPropagations(
+	const FObjectMixerEditorListRowData::FPropertyPropagationInfo& InPropagationInfo)
 {
-	for (const FObjectMixerEditorListRowPtr& SelectedItem : TreeViewPtr->GetSelectedItems())
-	{
-		SelectedItem->SetCurrentEditorObjectVisibility(bNewIsVisible, bIsRecursive);
-	}
+	PendingPropertyPropagations.Add(InPropagationInfo);
 }
 
-bool SObjectMixerEditorList::IsTreeViewItemSelected(TSharedRef<FObjectMixerEditorListRow> Item)
+TArray<TSharedPtr<ISceneOutlinerTreeItem>> SObjectMixerEditorList::GetSelectedTreeViewItems() const
 {
-	return TreeViewPtr->GetSelectedItems().Contains(Item);
+	auto Selection = GetTree().GetSelectedItems();
+	int32 Count = Selection.Num();
+	UE_LOG(LogTemp, Display, TEXT("%i"), Count);
+	return Selection;
 }
 
-TArray<FObjectMixerEditorListRowPtr> SObjectMixerEditorList::GetTreeViewItems() const
-{
-	return TreeViewRootObjects;
+int32 SObjectMixerEditorList::GetSelectedTreeViewItemCount() const
+{	
+	return GetTree().GetSelectedItems().Num();
 }
 
-void SObjectMixerEditorList::SetTreeViewItems(const TArray<FObjectMixerEditorListRowPtr>& InItems)
+void SObjectMixerEditorList::SetTreeViewItemSelected(TSharedRef<ISceneOutlinerTreeItem> Item, const bool bNewSelected)
 {
-	TreeViewRootObjects = InItems;
+	GetTreeView()->SetItemSelection(Item, bNewSelected);
+}
 
-	TreeViewPtr->RequestListRefresh();
+bool SObjectMixerEditorList::IsTreeViewItemSelected(TSharedPtr<ISceneOutlinerTreeItem> Item)
+{
+	if (const TSharedPtr<SSceneOutlinerTreeView> Tree = GetTreeView())
+	{
+		TArray<TSharedPtr<ISceneOutlinerTreeItem, ESPMode::ThreadSafe>> SelectedItems = Tree->GetSelectedItems();
+		return SelectedItems.Contains(Item);
+	}
+	
+	return false;
 }
 
-TSet<TWeakPtr<FObjectMixerEditorListRow>> SObjectMixerEditorList::GetSoloRows() const
+TSet<TSharedPtr<ISceneOutlinerTreeItem>> SObjectMixerEditorList::GetSoloRows() const
 {
-	TSet<TWeakPtr<FObjectMixerEditorListRow>> SoloRows;
+	TSet<TSharedPtr<ISceneOutlinerTreeItem>> SoloRows;
 	
-	using LambdaType = void(*)(const FObjectMixerEditorListRowPtr&, TSet<TWeakPtr<FObjectMixerEditorListRow>>&);
+	using LambdaType = void(*)(const TSharedPtr<ISceneOutlinerTreeItem>&, TSet<TSharedPtr<ISceneOutlinerTreeItem>>&);
 	static LambdaType GetSoloRowsRecursively =
-		[](const FObjectMixerEditorListRowPtr& InRow, TSet<TWeakPtr<FObjectMixerEditorListRow>>& SoloRows)
+		[](const TSharedPtr<ISceneOutlinerTreeItem>& InRow, TSet<TSharedPtr<ISceneOutlinerTreeItem>>& SoloRows)
 	{
-		if (InRow.IsValid())
+		if (FObjectMixerEditorListRowData* RowData = FObjectMixerUtils::GetRowData(InRow))
 		{
-			if (InRow->GetRowSoloState())
+			if (RowData->GetRowSoloState())
 			{
 				SoloRows.Add(InRow);
 			}
+		}
 			
-			for (const FObjectMixerEditorListRowPtr& Child : InRow->GetChildRows())
+		for (const TWeakPtr<ISceneOutlinerTreeItem>& Child : InRow->GetChildren())
+		{
+			if (Child.IsValid())
 			{
-				GetSoloRowsRecursively(Child, SoloRows);
+				GetSoloRowsRecursively(Child.Pin(), SoloRows);
 			}
 		}
 	};
 
-	for (const FObjectMixerEditorListRowPtr& RootRow : TreeViewRootObjects)
+	for (const TSharedPtr<ISceneOutlinerTreeItem>& RootRow : GetTreeRootItems())
 	{
 		GetSoloRowsRecursively(RootRow, SoloRows);
 	}
@@ -466,342 +484,203 @@ TSet<TWeakPtr<FObjectMixerEditorListRow>> SObjectMixerEditorList::GetSoloRows()
 	return SoloRows;
 }
 
-void SObjectMixerEditorList::ClearSoloRows()
+TSet<TSharedPtr<ISceneOutlinerTreeItem>> SObjectMixerEditorList::GetTreeRootItems() const
 {
-	using LambdaType = void(*)(const FObjectMixerEditorListRowPtr&);
-	static LambdaType ClearSoloRowsRecursively = [](const FObjectMixerEditorListRowPtr& InRow)
-	{
-		if (InRow.IsValid())
-		{
-			InRow->SetRowSoloState(false);
-			
-			for (const FObjectMixerEditorListRowPtr& Child : InRow->GetChildRows())
-			{
-				ClearSoloRowsRecursively(Child);
-			}
-		}
-	};
-
-	for (const FObjectMixerEditorListRowPtr& RootRow : TreeViewRootObjects)
+	if (GetTree().GetItems().Num())
 	{
-		ClearSoloRowsRecursively(RootRow);
+		return {GetTree().GetItems()[0]};
 	}
+
+	return {};
 }
 
-bool SObjectMixerEditorList::IsListInSoloState() const
+TSet<TWeakPtr<ISceneOutlinerTreeItem>> SObjectMixerEditorList::GetWeakTreeRootItems() const
 {
-	using LambdaType = bool(*)(const FObjectMixerEditorListRowPtr&);
-	static LambdaType GetSoloRowsRecursively = [](const FObjectMixerEditorListRowPtr& InRow)
+	if (GetTree().GetItems().Num())
 	{
-		if (InRow.IsValid())
-		{
-			if (InRow->GetRowSoloState())
-			{
-				return true;
-			}
-			
-			for (const FObjectMixerEditorListRowPtr& Child : InRow->GetChildRows())
-			{
-				if (GetSoloRowsRecursively(Child))
-				{
-					return true;
-				}
-			}
-		}
-
-		return false;
-	};
-
-	for (const FObjectMixerEditorListRowPtr& RootRow : TreeViewRootObjects)
-	{
-		if (GetSoloRowsRecursively(RootRow))
-		{
-			return true;
-		}
+		return {GetTree().GetItems()[0]};
 	}
 
-	return false;
+	return {};
 }
 
-void SObjectMixerEditorList::EvaluateAndSetEditorVisibilityPerRow()
+void SObjectMixerEditorList::ClearSoloRows()
 {
-	using LambdaType = void(*)(const FObjectMixerEditorListRowPtr&, const bool);
-	static LambdaType EvaluateEditorVisibilityRecursively =
-		[](const FObjectMixerEditorListRowPtr& InRow, const bool bIsListInSoloState)
+	using LambdaType = void(*)(const TSharedPtr<ISceneOutlinerTreeItem>&);
+	static LambdaType ClearSoloRowsRecursively = [](const TSharedPtr<ISceneOutlinerTreeItem>& InRow)
 	{
-		if (InRow.IsValid())
+		if (FObjectMixerEditorListRowData* RowData = FObjectMixerUtils::GetRowData(InRow))
 		{
-			bool bShouldBeVisible = !InRow->IsUserSetHiddenInEditor();
-
-			if (bIsListInSoloState)
-			{
-				bShouldBeVisible = InRow->GetRowSoloState();
-			}
-
-			InRow->SetCurrentEditorObjectVisibility(bShouldBeVisible);
+			RowData->SetRowSoloState(false);
+		}
 			
-			for (const FObjectMixerEditorListRowPtr& Child : InRow->GetChildRows())
+		for (const TWeakPtr<ISceneOutlinerTreeItem>& Child : InRow->GetChildren())
+		{
+			if (Child.IsValid())
 			{
-				EvaluateEditorVisibilityRecursively(Child, bIsListInSoloState);
+				ClearSoloRowsRecursively(Child.Pin());
 			}
 		}
 	};
-	
-	const bool bIsListInSoloState = IsListInSoloState();
 
-	for (const FObjectMixerEditorListRowPtr& RootRow : TreeViewRootObjects)
+	for (const TSharedPtr<ISceneOutlinerTreeItem>& RootItem : GetTreeRootItems())
 	{
-		EvaluateEditorVisibilityRecursively(RootRow, bIsListInSoloState);
+		ClearSoloRowsRecursively(RootItem);
 	}
 }
 
-FText SObjectMixerEditorList::GetSearchTextFromSearchInputField() const
-{
-	TSharedPtr<FObjectMixerEditorList> PinnedListModel = ListModelPtr.Pin();
-	check(PinnedListModel && PinnedListModel->GetMainPanelModel().IsValid());
-	check(HeaderRow);
-	
-	return PinnedListModel->GetMainPanelModel().Pin()->GetSearchTextFromSearchInputField();
-}
-
-FString SObjectMixerEditorList::GetSearchStringFromSearchInputField() const
-{
-	TSharedPtr<FObjectMixerEditorList> PinnedListModel = ListModelPtr.Pin();
-	check(PinnedListModel && PinnedListModel->GetMainPanelModel().IsValid());
-	check(HeaderRow);
-	
-	return PinnedListModel->GetMainPanelModel().Pin()->GetSearchStringFromSearchInputField();
-}
-
-void SObjectMixerEditorList::ExecuteListViewSearchOnAllRows(
-	const FString& SearchString, const bool bShouldRefreshAfterward)
+bool SObjectMixerEditorList::IsListInSoloState() const
 {
-	TArray<FString> Tokens;
-	
-	// unquoted search equivalent to a match-any-of search
-	SearchString.ParseIntoArray(Tokens, TEXT("|"), true);
-	
-	for (const FObjectMixerEditorListRowPtr& ChildRow : TreeViewRootObjects)
+	using LambdaType = bool(*)(const TSharedPtr<ISceneOutlinerTreeItem>&);
+	static LambdaType GetSoloRowsRecursively = [](const TSharedPtr<ISceneOutlinerTreeItem>& InRow)
 	{
-		if (!ensure(ChildRow.IsValid()))
+		if (FObjectMixerEditorListRowData* RowData = FObjectMixerUtils::GetRowData(InRow))
 		{
-			continue;
+			if (RowData->GetRowSoloState())
+			{
+				return true;
+			}
 		}
-		
-		const bool bGroupMatch = ChildRow->MatchSearchTokensToSearchTerms(Tokens);
-		
-		// If the group name matches then we pass in an empty string so all child nodes are visible.
-		// If the name doesn't match, then we need to evaluate each child.
-		ChildRow->ExecuteSearchOnChildNodes(bGroupMatch ? "" : SearchString);
-	}
-
-	if (bShouldRefreshAfterward)
-	{
-		FindVisibleObjectsAndRequestTreeRefresh();
-	}
-}
-
-bool SObjectMixerEditorList::DoesTreeViewHaveVisibleChildren() const
-{
-	if (TreeViewPtr.IsValid())
-	{
-		for (const FObjectMixerEditorListRowPtr& Header : TreeViewRootObjects)
-		{
-			const EVisibility HeaderVisibility = Header->GetDesiredRowWidgetVisibility();
 			
-			if (HeaderVisibility != EVisibility::Hidden && HeaderVisibility != EVisibility::Collapsed)
+		for (const TWeakPtr<ISceneOutlinerTreeItem>& Child : InRow->GetChildren())
+		{
+			if (Child.IsValid())
 			{
-				return true;
+				if (GetSoloRowsRecursively(Child.Pin()))
+				{
+					return true;
+				}
 			}
 		}
-	}
-	
-	return false;
-}
 
-bool SObjectMixerEditorList::IsTreeViewItemExpanded(const FObjectMixerEditorListRowPtr& Row) const
-{
-	if (TreeViewPtr.IsValid())
-	{
-		return TreeViewPtr->IsItemExpanded(Row);
-	}
-	
-	return false;
-}
+		return false;
+	};
 
-void SObjectMixerEditorList::SetTreeViewItemExpanded(const FObjectMixerEditorListRowPtr& RowToExpand, const bool bNewExpansion) const
-{
-	if (TreeViewPtr.IsValid())
+	for (const TSharedPtr<ISceneOutlinerTreeItem>& RootItem : GetTreeRootItems())
 	{
-		TreeViewPtr->SetItemExpansion(RowToExpand, bNewExpansion);
+		if (GetSoloRowsRecursively(RootItem))
+		{
+			return true;
+		}
 	}
-}
-
-EObjectMixerTreeViewMode SObjectMixerEditorList::GetTreeViewMode()
-{
-	const TSharedPtr<FObjectMixerEditorList> PinnedListModel = GetListModelPtr().Pin();
-	check(PinnedListModel);
-
-	const TSharedPtr<FObjectMixerEditorMainPanel> PinnedMainPanel = PinnedListModel->GetMainPanelModel().Pin();
-	check(PinnedMainPanel);
-	
-	return PinnedMainPanel->GetTreeViewMode();
-}
 
-const TArray<TSharedRef<IObjectMixerEditorListFilter>>& SObjectMixerEditorList::GetListFilters()
-{
-	const TSharedPtr<FObjectMixerEditorList> PinnedListModel = GetListModelPtr().Pin();
-	check(PinnedListModel);
-
-	const TSharedPtr<FObjectMixerEditorMainPanel> PinnedMainPanel = PinnedListModel->GetMainPanelModel().Pin();
-	check(PinnedMainPanel);
-	
-	return PinnedMainPanel->GetListFilters();
+	return false;
 }
 
-void SObjectMixerEditorList::EvaluateIfRowsPassFilters(const bool bShouldRefreshAfterward)
+void SObjectMixerEditorList::EvaluateAndSetEditorVisibilityPerRow()
 {
-	// Separate filters by type
-	
-	TSet<TSharedRef<IObjectMixerEditorListFilter>> MatchAnyOfFilters;
-	TSet<TSharedRef<IObjectMixerEditorListFilter>> MatchAllOfFilters;
-
-	for (const TSharedRef<IObjectMixerEditorListFilter>& Filter : GetListFilters())
+	using LambdaType = void(*)(const TSharedPtr<ISceneOutlinerTreeItem>&, const bool);
+	static LambdaType EvaluateEditorVisibilityRecursively =
+		[](const TSharedPtr<ISceneOutlinerTreeItem>& InRow, const bool bIsListInSoloState)
 	{
-		if (Filter->GetFilterMatchType() ==
-			IObjectMixerEditorListFilter::EObjectMixerEditorListFilterMatchType::MatchAll)
-		{
-			MatchAllOfFilters.Add(Filter);
-		}
-		else
+		if (const FObjectMixerEditorListRowActor* ActorRow = FObjectMixerUtils::AsActorRow(InRow))
 		{
-			MatchAnyOfFilters.Add(Filter);
-		}
-	}
+			bool bShouldBeHidden = ActorRow->RowData.IsUserSetHiddenInEditor();
 
-	const TFunction<void(const TArray<FObjectMixerEditorListRowPtr>&)> RecursiveRowIterator =
-		[&MatchAnyOfFilters, &MatchAllOfFilters, &RecursiveRowIterator](const TArray<FObjectMixerEditorListRowPtr>& InObjects)
-	{
-		bool bExpandByDefault = true;
-		if (const UObjectMixerEditorSettings* Settings = GetDefault<UObjectMixerEditorSettings>())
-		{
-			bExpandByDefault = Settings->bExpandTreeViewItemsByDefault;
+			if (bIsListInSoloState)
+			{
+				bShouldBeHidden = !ActorRow->RowData.GetRowSoloState();
+			}
+
+			// Only create a transaction in the case of a change
+			if (bShouldBeHidden != ActorRow->Actor->IsTemporarilyHiddenInEditor())
+			{
+				// Save the actor to the transaction buffer to support undo/redo, but do
+				// not call Modify, as we do not want to dirty the actor's package and
+				// we're only editing temporary, transient values
+				SaveToTransactionBuffer(ActorRow->Actor.Get(), false);
+				ActorRow->Actor->SetIsTemporarilyHiddenInEditor(bShouldBeHidden);
+			}
 		}
 			
-		for (const FObjectMixerEditorListRowPtr& Row : InObjects)
+		for (const TWeakPtr<ISceneOutlinerTreeItem>& Child : InRow->GetChildren())
 		{
-			if (Row.IsValid())
+			if (Child.IsValid())
 			{
-				auto FilterProjection = [&Row](const TSharedRef<IObjectMixerEditorListFilter>& Filter)
-				{
-					return Filter->GetIsFilterActive() ? Filter->DoesItemPassFilter(Row) : true;
-				};
-			
-				const bool bPassesAnyOf = MatchAnyOfFilters.Num() ? Algo::AnyOf(MatchAnyOfFilters, FilterProjection) : true;
-				const bool bPassesAllOf = MatchAllOfFilters.Num() ? Algo::AllOf(MatchAllOfFilters, FilterProjection) : true;
-		
-				Row->SetDoesRowPassFilters(bPassesAnyOf && bPassesAllOf);
-				
-				if (Row->GetChildCount() > 0)
-				{
-					RecursiveRowIterator(Row->GetChildRows());
-				}
-				
-				Row->SetIsTreeViewItemExpanded(bExpandByDefault);
+				EvaluateEditorVisibilityRecursively(Child.Pin(), bIsListInSoloState);
 			}
 		}
 	};
 
-	RecursiveRowIterator(TreeViewRootObjects);
-
-	if (bShouldRefreshAfterward)
-	{
-		FindVisibleObjectsAndRequestTreeRefresh();
-	}
-}
-
-EColumnSortMode::Type SObjectMixerEditorList::GetSortModeForColumn(FName InColumnName) const
-{
-	EColumnSortMode::Type ColumnSortMode = EColumnSortMode::None;
+	const bool bIsListInSoloState = IsListInSoloState();
 
-	if (GetActiveSortingColumnName().IsEqual(InColumnName))
+	for (const TSharedPtr<ISceneOutlinerTreeItem>& RootItem : GetTreeRootItems())
 	{
-		ColumnSortMode = ActiveSortingType;
+		EvaluateEditorVisibilityRecursively(RootItem, bIsListInSoloState);
 	}
-
-	return ColumnSortMode;
 }
 
-void SObjectMixerEditorList::OnSortColumnCalled(EColumnSortPriority::Type Priority, const FName& ColumnName, EColumnSortMode::Type SortMode)
+bool SObjectMixerEditorList::IsTreeViewItemExpanded(const TSharedPtr<ISceneOutlinerTreeItem>& Row) const
 {
-	ExecuteSort(ColumnName, CycleSortMode(ColumnName));
+	return GetTree().IsItemExpanded(Row);
 }
 
-EColumnSortMode::Type SObjectMixerEditorList::CycleSortMode(const FName& InColumnName)
+void SObjectMixerEditorList::SetTreeViewItemExpanded(const TSharedPtr<ISceneOutlinerTreeItem>& RowToExpand, const bool bNewExpansion) const
 {
-	const EColumnSortMode::Type PreviousColumnSortMode = GetSortModeForColumn(InColumnName);
-	ActiveSortingType = PreviousColumnSortMode ==
-		EColumnSortMode::Ascending ? EColumnSortMode::Descending : EColumnSortMode::Ascending;
-
-	ActiveSortingColumnName = InColumnName;
-	return ActiveSortingType;
-}
-
-void SObjectMixerEditorList::ExecuteSort(
-	const FName& InColumnName, const EColumnSortMode::Type InColumnSortMode, const bool bShouldRefreshAfterward)
-{	
-	if (bShouldRefreshAfterward)
-	{
-		FindVisibleObjectsAndRequestTreeRefresh();
-	}
+	GetTreeView()->SetItemExpansion(RowToExpand, bNewExpansion);
 }
 
-FListViewColumnInfo* SObjectMixerEditorList::GetColumnInfoByPropertyName(const FName& InPropertyName)
+FObjectMixerSceneOutlinerColumnInfo* SObjectMixerEditorList::GetColumnInfoByPropertyName(const FName& InPropertyName)
 {
-	return Algo::FindByPredicate(ListViewColumns,
-		[InPropertyName] (const FListViewColumnInfo& ColumnInfo)
+	return Algo::FindByPredicate(HeaderColumnInfos,
+		[InPropertyName] (const FObjectMixerSceneOutlinerColumnInfo& ColumnInfo)
 		{
 			return ColumnInfo.PropertyName.IsEqual(InPropertyName);
 		});
 }
 
-void SObjectMixerEditorList::ExecuteRenameOnPending()
+void SObjectMixerEditorList::CustomAddToToolbar(TSharedPtr<SHorizontalBox> Toolbar)
 {
-	using LambdaType = const FObjectMixerEditorListRowPtr*(*)(const TArray<FObjectMixerEditorListRowPtr>&, const FObjectMixerEditorListRow::FTreeItemUniqueIdentifier&);
-		
-	static LambdaType RecursivelyFindByUniqueIdentifier = [](
-		const TArray<FObjectMixerEditorListRowPtr>& InObjects, const FObjectMixerEditorListRow::FTreeItemUniqueIdentifier& InUniqueIdentifier)
-	{
-		for (const FObjectMixerEditorListRowPtr& TreeViewItem : InObjects)
-		{
-			const FObjectMixerEditorListRow::FTreeItemUniqueIdentifier& TreeItemUniqueIdentifier = TreeViewItem->GetUniqueIdentifier();
+	ToolbarPtr = Toolbar;
+
+	// "Add object" button
+	Toolbar->InsertSlot(0)
+	.HAlign(HAlign_Left)
+	.VAlign(VAlign_Center)
+	.AutoWidth()
+	.Padding(FMargin(0, 4, 8, 4))
+	[
+		SNew(SPositiveActionButton)
+		.Text(LOCTEXT("AddObject", "Add"))
+		.OnGetMenuContent(FOnGetContent::CreateRaw(this, &SObjectMixerEditorList::OnGenerateAddObjectButtonMenu))
+	];
 
-			if (InUniqueIdentifier == TreeItemUniqueIdentifier)
+	// Sync selection toggle
+	Toolbar->AddSlot()
+	.HAlign(HAlign_Right)
+	.VAlign(VAlign_Center)
+	.AutoWidth()
+	.Padding(8.f, 1.f, 0.f, 1.f)
+	[
+		SNew(SCheckBox )
+		.Padding(4.f)
+		.ToolTipText(LOCTEXT("SyncSelectionButton_Tooltip", "Sync Selection\nIf enabled, clicking an item in the mixer list will also select the item in the Scene Outliner.\nAlt + Click to select items in mixer without selecting the item in the Scene outliner.\nIf disabled, selections will not sync unless Alt is held. Effectively, this is the opposite behavior."))
+		.Style(&FAppStyle::Get().GetWidgetStyle<FCheckBoxStyle>("ToggleButtonCheckbox"))
+		.IsChecked_Lambda([]()
+		{
+			if (const UObjectMixerEditorSettings* Settings = GetDefault<UObjectMixerEditorSettings>())
 			{
-				return &TreeViewItem;
+				return Settings->bSyncSelection ? ECheckBoxState::Checked : ECheckBoxState::Unchecked;
 			}
-
-			if (const FObjectMixerEditorListRowPtr* FoundItem = RecursivelyFindByUniqueIdentifier(
-				TreeViewItem->GetChildRows(), InUniqueIdentifier))
+			
+			return ECheckBoxState::Undetermined; 
+		})
+		.OnCheckStateChanged_Lambda([](ECheckBoxState InNewState)
+		{
+			if (UObjectMixerEditorSettings* Settings = GetMutableDefault<UObjectMixerEditorSettings>())
 			{
-				return FoundItem;
+				Settings->bSyncSelection = InNewState == ECheckBoxState::Checked ? true : false;
 			}
-		}
-
-		const FObjectMixerEditorListRowPtr* Nulled = nullptr;
-		return Nulled;
-	};
-
-	if (const FObjectMixerEditorListRowPtr* FoundItem = RecursivelyFindByUniqueIdentifier(TreeViewRootObjects, PendingRenameItem))
-	{
-		(*FoundItem)->CallOnRenameCommandDelegate();
-	}
-
-	PendingRenameItem.Reset();
+		})
+		[
+			SNew(SImage)
+			.ColorAndOpacity(FSlateColor::UseForeground())
+			.Image( FAppStyle::Get().GetBrush("FoliageEditMode.SelectAll") )
+		]
+	];
 }
 
-void SObjectMixerEditorList::Tick(const FGeometry& AllottedGeometry, const double InCurrentTime,
-                                  const float InDeltaTime)
+void SObjectMixerEditorList::Tick(const FGeometry& AllottedGeometry, const double InCurrentTime, const float InDeltaTime)
 {
 	SCompoundWidget::Tick(AllottedGeometry, InCurrentTime, InDeltaTime);
 
@@ -810,29 +689,16 @@ void SObjectMixerEditorList::Tick(const FGeometry& AllottedGeometry, const doubl
 		RebuildList();
 		
 		// Early out so we have widgets to act upon next frame
+		SSceneOutliner::Tick(AllottedGeometry, InCurrentTime, InDeltaTime);
 		return;
 	}
 
-	if (bIsEditorToListSelectionSyncRequested)
-	{
-		SyncEditorSelectionToListSelection();
-	}
-
-	if (bShouldPauseSyncSelection)
-	{
-		bShouldPauseSyncSelection = false;
-		bIsEditorToListSelectionSyncRequested = false;
-	}
-
 	if (PendingPropertyPropagations.Num() > 0)
 	{
 		PropagatePropertyChangesToSelectedRows();
 	}
 
-	if (PendingRenameItem.IdentifiesAnyRow())
-	{
-		ExecuteRenameOnPending();
-	}
+	SSceneOutliner::Tick(AllottedGeometry, InCurrentTime, InDeltaTime);
 }
 
 TSharedRef<SWidget> SObjectMixerEditorList::GenerateHeaderRowContextMenu() const
@@ -843,7 +709,7 @@ TSharedRef<SWidget> SObjectMixerEditorList::GenerateHeaderRowContextMenu() const
 
 	FName LastPropertyCategoryName = NAME_None;
 
-	for (const FListViewColumnInfo& ColumnInfo : ListViewColumns)
+	for (const FObjectMixerSceneOutlinerColumnInfo& ColumnInfo : HeaderColumnInfos)
 	{
 		const FName& PropertyCategoryName = ColumnInfo.PropertyCategoryName;
 
@@ -855,12 +721,12 @@ TSharedRef<SWidget> SObjectMixerEditorList::GenerateHeaderRowContextMenu() const
             MenuBuilder.BeginSection(LastPropertyCategoryName, FText::FromName(LastPropertyCategoryName));
 		}
 		
-		const FName& PropertyName = ColumnInfo.PropertyName;
+		const FName& ColumnID = ColumnInfo.ColumnID;
 		
 		const FText Tooltip = ColumnInfo.PropertyRef ?
 			ColumnInfo.PropertyRef->GetToolTipText() : ColumnInfo.PropertyDisplayText;
 
-		const bool bCanSelectColumn = ColumnInfo.PropertyType != EListViewColumnType::BuiltIn;
+		const bool bCanSelectColumn = ColumnInfo.bCanBeHidden;
 
 		const FName Hook = ColumnInfo.PropertyType == EListViewColumnType::BuiltIn ? "Builtin" : "GeneratedProperties";
 		
@@ -869,31 +735,29 @@ TSharedRef<SWidget> SObjectMixerEditorList::GenerateHeaderRowContextMenu() const
 			Tooltip,
 			FSlateIcon(),
 			FUIAction(
-				FExecuteAction::CreateLambda([this, PropertyName]()
+				FExecuteAction::CreateLambda([this, ColumnID]()
 				{
 					TSharedPtr<FObjectMixerEditorList> PinnedListModel = ListModelPtr.Pin();
-					check(PinnedListModel && PinnedListModel->GetMainPanelModel().IsValid());
-					check(HeaderRow);
+					TSharedPtr<SHeaderRow> PinnedHeaderRow = GetTree().GetHeaderRow();
+					check(PinnedHeaderRow);
 
-					const bool bNewColumnEnabled = !HeaderRow->IsColumnVisible(PropertyName);
+					const bool bNewColumnEnabled = !PinnedHeaderRow->IsColumnVisible(ColumnID);
 					
-					HeaderRow->SetShowGeneratedColumn(PropertyName, bNewColumnEnabled);
+					PinnedHeaderRow->SetShowGeneratedColumn(ColumnID, bNewColumnEnabled);
 
 					if (UObjectMixerEditorSerializedData* SerializedData = GetMutableDefault<UObjectMixerEditorSerializedData>())
 					{
-						for (const TSubclassOf<UObjectMixerObjectFilter>& FilterClass : PinnedListModel->GetMainPanelModel().Pin()->GetObjectFilterClasses())
+						for (const TSubclassOf<UObjectMixerObjectFilter>& FilterClass : PinnedListModel->GetObjectFilterClasses())
 						{
 							const FName FilterName = FilterClass->GetFName();
-							SerializedData->SetShouldShowColumn(FilterName, PropertyName, bNewColumnEnabled);
+							SerializedData->SetShouldShowColumn(FilterName, ColumnID, bNewColumnEnabled);
 						}
 					}
 				}),
 				FCanExecuteAction::CreateLambda([bCanSelectColumn](){return bCanSelectColumn;}),
-				FIsActionChecked::CreateLambda([this, PropertyName]()
+				FIsActionChecked::CreateLambda([this, ColumnID]()
 				{
-					check(HeaderRow);
-					
-					return HeaderRow->IsColumnVisible(PropertyName);
+					return GetTree().GetHeaderRow()->IsColumnVisible(ColumnID);
 				})
 			),
 			Hook,
@@ -904,7 +768,7 @@ TSharedRef<SWidget> SObjectMixerEditorList::GenerateHeaderRowContextMenu() const
 	return MenuBuilder.MakeWidget();
 }
 
-bool SObjectMixerEditorList::AddUniquePropertyColumnsToHeaderRow(
+bool SObjectMixerEditorList::AddUniquePropertyColumnInfo(
 	FProperty* Property,
 	const bool bForceIncludeProperty,
 	const TSet<FName>& PropertySkipList)
@@ -943,8 +807,8 @@ bool SObjectMixerEditorList::AddUniquePropertyColumnsToHeaderRow(
 		const FName PropertyName = Property->GetFName();
 	
 		// Ensure no duplicate properties
-		if (!Algo::FindByPredicate(ListViewColumns,
-				[&PropertyName] (const FListViewColumnInfo& ListViewColumn)
+		if (!Algo::FindByPredicate(HeaderColumnInfos,
+				[&PropertyName] (const FObjectMixerSceneOutlinerColumnInfo& ListViewColumn)
 				{
 					return ListViewColumn.PropertyName.IsEqual(PropertyName);
 				})
@@ -956,13 +820,13 @@ bool SObjectMixerEditorList::AddUniquePropertyColumnsToHeaderRow(
 				PropertyCategoryName = *CategoryMeta;
 			}
 			
-			ListViewColumns.Add(
+			HeaderColumnInfos.Add(
 				{
-					Property, PropertyName,
+					Property, PropertyName, FObjectMixerOutlinerPropertyColumn::GetID(Property),
 					Property->GetDisplayNameText(),
 					EListViewColumnType::PropertyGenerated,
 					*PropertyCategoryName,
-					true, false, false
+					true, true, false, false
 				}
 			);
 
@@ -973,803 +837,582 @@ bool SObjectMixerEditorList::AddUniquePropertyColumnsToHeaderRow(
 	return false;
 }
 
-void SObjectMixerEditorList::AddBuiltinColumnsToHeaderRow()
-{
-	ListViewColumns.Insert(
-		{
-			nullptr, ItemNameColumnName,
-			LOCTEXT("ItemNameHeaderText", "Name"),
-			EListViewColumnType::BuiltIn, "Built-In",
-			true, false,
-			false, 1.0f, 1.7f
-		}, 0
-	);
-	
-	ListViewColumns.Insert(
-		{
-			nullptr, EditorVisibilitySoloColumnName,
-			LOCTEXT("EditorVisibilitySoloColumnNameHeaderText", "Solo"),
-			EListViewColumnType::BuiltIn, "Built-In",
-			true, false,
-			true, 25.0f
-		}, 0
-	);
-
-	ListViewColumns.Insert(
-		{
-			nullptr, EditorVisibilityColumnName,
-			LOCTEXT("EditorVisibilityColumnNameHeaderText", "Visibility"),
-			EListViewColumnType::BuiltIn, "Built-In",
-			true, false,
-			true, 25.0f
-		}, 0
-	);
-}
-
-TSharedPtr<SHeaderRow> SObjectMixerEditorList::GenerateHeaderRow()
+void SObjectMixerEditorList::CreateActorTextInfoColumns(UWorld *WorldPtr, FSceneOutlinerInitializationOptions& OutInitOptions)
 {
-	TSharedPtr<FObjectMixerEditorList> PinnedListModel = ListModelPtr.Pin();
-	check(PinnedListModel && PinnedListModel->GetMainPanelModel().IsValid());
-	check(HeaderRow);
-
-	HeaderRow->ClearColumns();
-	ListViewColumns.Empty(ListViewColumns.Num());
-	
-	TSet<UClass*> SpecifiedClasses;
-	for (const TObjectPtr<UObjectMixerObjectFilter>& Instance : PinnedListModel->GetMainPanelModel().Pin()->GetObjectFilterInstances())
-	{
-		SpecifiedClasses.Append(
-			Instance->GetParentAndChildClassesFromSpecifiedClasses(
-				ObjectClassesToFilterCache, PropertyInheritanceInclusionOptionsCache
-			)
-		);
-	}
-	
-	for (const UClass* Class : SpecifiedClasses)
+	FGetTextForItem LayerInfoText = FGetTextForItem::CreateLambda([](const ISceneOutlinerTreeItem& Item) -> FString
 	{
-		for (TFieldIterator<FProperty> FieldIterator(Class, EFieldIterationFlags::None); FieldIterator; ++FieldIterator)
+		FString Result;
+		
+		if (const FActorTreeItem* ActorItem = Item.CastTo<FActorTreeItem>())
 		{
-			if (FProperty* Property = *FieldIterator)
+			AActor* Actor = ActorItem->Actor.Get();
+			if (!Actor)
 			{
-				AddUniquePropertyColumnsToHeaderRow(Property, bShouldIncludeUnsupportedPropertiesCache, ColumnsToExcludeCache);
+				return FString();
 			}
-		}
 
-		// Check Force Added Columns
-		for (const FName& PropertyName : ForceAddedColumnsCache)
-		{
-			if (FProperty* Property = FindFProperty<FProperty>(Class, PropertyName))
+			for (const auto& Layer : Actor->Layers)
 			{
-				AddUniquePropertyColumnsToHeaderRow(Property, true);
+				if (Result.Len())
+				{
+					Result += TEXT(", ");
+				}
+
+				Result += Layer.ToString();
 			}
 		}
-	}
 
-	// Alphabetical sort by Property Name
-	ListViewColumns.StableSort([](const FListViewColumnInfo& A, const FListViewColumnInfo& B)
-	{
-		return A.PropertyDisplayText.ToString() < B.PropertyDisplayText.ToString();
+		return Result;
 	});
 
-	// Alphabetical sort by Property Category Name
-	ListViewColumns.StableSort([](const FListViewColumnInfo& A, const FListViewColumnInfo& B)
+	FGetTextForItem DataLayerInfoText = FGetTextForItem::CreateLambda([](const ISceneOutlinerTreeItem& Item) -> FString
 	{
-		return A.PropertyCategoryName.LexicalLess(B.PropertyCategoryName);
-	});
-
-	// Add Built-in Columns to beginning
-	AddBuiltinColumnsToHeaderRow();
+		TStringBuilder<128> Builder;
+		TSet<FString> DataLayerShortNames;
 
-	// Actually add columns to Header
-	{	
-		const FText ClickToSortTooltip = LOCTEXT("ClickToSort","Click to sort");
-
-		const TSharedRef<SWidget> HeaderMenuContent = GenerateHeaderRowContextMenu();
-
-		for (const FListViewColumnInfo& ColumnInfo : ListViewColumns)
+		auto BuildDataLayers = [&Builder, &DataLayerShortNames](const auto& DataLayerInstances, bool bPartOfOtherLevel)
 		{
-			const FText Tooltip = ColumnInfo.PropertyRef ? ColumnInfo.PropertyRef->GetToolTipText() :
-				ColumnInfo.bCanBeSorted ? ClickToSortTooltip : ColumnInfo.PropertyDisplayText;
-		
-			SHeaderRow::FColumn::FArguments Column =
-				SHeaderRow::Column(ColumnInfo.PropertyName)
-				.DefaultLabel(ColumnInfo.PropertyDisplayText)
-				.ToolTipText(Tooltip)
-				.HAlignHeader(EHorizontalAlignment::HAlign_Left)
-				//.bHideHeaderMenuButton(false) todo: uncomment this line when 21478974 is pushed 
-			;
-
-			if (ColumnInfo.bUseFixedWidth)
+			for (const UDataLayerInstance* DataLayerInstance : DataLayerInstances)
 			{
-				Column.FixedWidth(ColumnInfo.FixedWidth);
-			}
-			else
-			{
-				Column.FillWidth(ColumnInfo.FillWidth);
-			}
-
-			if (ColumnInfo.bCanBeSorted)
-			{
-				Column.SortMode_Raw(this, &SObjectMixerEditorList::GetSortModeForColumn, ColumnInfo.PropertyName);
-				Column.OnSort_Raw(this, &SObjectMixerEditorList::OnSortColumnCalled);
-			}
-
-			if (ColumnInfo.PropertyType == EListViewColumnType::BuiltIn)
-			{
-				Column.ShouldGenerateWidget(true);
+				bool bIsAlreadyInSet = false;
+				DataLayerShortNames.Add(DataLayerInstance->GetDataLayerShortName(), &bIsAlreadyInSet);
+				if (!bIsAlreadyInSet)
+				{
+					if (Builder.Len())
+					{
+						Builder += TEXT(", ");
+					}
+					// Put a '*' in front of DataLayers that are not part of of the main world
+					if (bPartOfOtherLevel)
+					{
+						Builder += "*";
+					}
+					Builder += DataLayerInstance->GetDataLayerShortName();
+				}
 			}
+		};
 
-			if (ColumnInfo.PropertyName.IsEqual(EditorVisibilityColumnName))
-			{
-				Column.HeaderContent()
-				[
-					SNew(SBox)
-					.HAlign(HAlign_Left)
-					.VAlign(VAlign_Center)
-					.Padding(0.f)
-					[
-						SNew(SImage)
-						.ColorAndOpacity(FSlateColor::UseForeground())
-						.Image(FAppStyle::Get().GetBrush("Level.VisibleIcon16x"))
-						.ToolTipText(LOCTEXT("VisibilityColumnTooltip", "Visibility"))
-					]
-				];
-			}
-			else if (ColumnInfo.PropertyName.IsEqual(EditorVisibilitySoloColumnName))
-			{
-				Column.HeaderContent()
-				[
-					SNew(SBox)
-					.HAlign(HAlign_Left)
-					.VAlign(VAlign_Center)
-					.Padding(0.f)
-					[
-						SNew(SImage)
-						.ColorAndOpacity(FSlateColor::UseForeground())
-						.Image(FObjectMixerEditorStyle::Get().GetBrush("ObjectMixer.Solo"))
-						.ToolTipText(LOCTEXT("SoloColumnTooltip", "Solo"))
-					]
-				];
-			}
-			else // Add Column Selection Menu widget to all other columns
+		auto BuildDataLayersWithContext = [BuildDataLayers](const ISceneOutlinerTreeItem& Item, bool bUseLevelContext)
+		{
+			if (const FActorTreeItem* ActorItem = Item.CastTo<FActorTreeItem>())
 			{
-				Column.MenuContent()
-				[
-					HeaderMenuContent
-				];
+				if (AActor* Actor = ActorItem->Actor.Get())
+				{
+					BuildDataLayers(bUseLevelContext ? Actor->GetDataLayerInstancesForLevel() : Actor->GetDataLayerInstances(), bUseLevelContext);
+				}
 			}
-		
-			HeaderRow->AddColumn(Column);
-
-			// Figure out if we should show a certain column
-			// First check filter rules
-			bool bShouldShowColumn = ColumnsToShowByDefaultCache.Contains(ColumnInfo.PropertyName);
-
-			// Then load visible columns from SerializedData
-			if (UObjectMixerEditorSerializedData* SerializedData = GetMutableDefault<UObjectMixerEditorSerializedData>())
+			else if (const FActorDescTreeItem* ActorDescItem = Item.CastTo<FActorDescTreeItem>())
 			{
-				for (const TSubclassOf<UObjectMixerObjectFilter>& Filter : PinnedListModel->GetMainPanelModel().Pin()->GetObjectFilterClasses())
+				if (const FWorldPartitionActorDesc* ActorDesc = ActorDescItem->ActorDescHandle.Get(); ActorDesc && !ActorDesc->GetDataLayerInstanceNames().IsEmpty())
 				{
-					const FName FilterName = Filter->GetFName();
-					if (SerializedData->IsColumnDataSerialized(FilterName, ColumnInfo.PropertyName))
+					if (const UActorDescContainer* ActorDescContainer = ActorDescItem->ActorDescHandle.Container.Get())
 					{
-						if (SerializedData->ShouldShowColumn(FilterName, ColumnInfo.PropertyName))
+						const UWorld* World = ActorDescContainer->GetWorld();
+						if (const UDataLayerSubsystem* DataLayerSubsystem = UWorld::GetSubsystem<UDataLayerSubsystem>(World))
 						{
-							// If any filter's serialized data says we should show a column, show it and stomp other filters' saved setting.
-							bShouldShowColumn = true;
-							break;
+							TSet<const UDataLayerInstance*> DataLayerInstances;
+							DataLayerInstances.Append(DataLayerSubsystem->GetDataLayerInstances(ActorDesc->GetDataLayerInstanceNames()));
+							if (ULevelInstanceSubsystem* LevelInstanceSubsystem = UWorld::GetSubsystem<ULevelInstanceSubsystem>(World))
+							{
+								UWorld* OuterWorld = ActorDescContainer->GetTypedOuter<UWorld>();
+								// Add parent container Data Layer Instances
+								AActor* CurrentActor = OuterWorld ? Cast<AActor>(LevelInstanceSubsystem->GetOwningLevelInstance(OuterWorld->PersistentLevel)) : nullptr;
+								while (CurrentActor)
+								{
+									DataLayerInstances.Append(bUseLevelContext ? CurrentActor->GetDataLayerInstancesForLevel() : CurrentActor->GetDataLayerInstances());
+									CurrentActor = Cast<AActor>(LevelInstanceSubsystem->GetParentLevelInstance(CurrentActor));
+								};
+							}
+							BuildDataLayers(DataLayerInstances, bUseLevelContext);
 						}
 					}
 				}
 			}
-		
-			HeaderRow->SetShowGeneratedColumn(ColumnInfo.PropertyName, bShouldShowColumn);
-		}
-	}
+		};
 
-	return HeaderRow;
-}
+		// List Actor's DataLayers part of the owning world, then those only part of the actor level
+		BuildDataLayersWithContext(Item, false);
+		BuildDataLayersWithContext(Item, true);
 
-void SObjectMixerEditorList::FlushMemory(const bool bShouldKeepMemoryAllocated)
-{
-	if (bShouldKeepMemoryAllocated)
-	{
-		ObjectsToRowsCreated.Reset();
-		TreeViewRootObjects.Reset();
-		VisibleTreeViewObjects.Reset();
-	}
-	else
-	{
-		ObjectsToRowsCreated.Empty();
-		TreeViewRootObjects.Empty();
-		VisibleTreeViewObjects.Empty();
-	}
-}
+		return Builder.ToString();
+	});
 
-void SObjectMixerEditorList::SetAllGroupsCollapsed()
-{
-	if (TreeViewPtr.IsValid())
-	{
-		for (const FObjectMixerEditorListRowPtr& RootRow : TreeViewRootObjects)
+	FGetTextForItem ContentBundleInfoText = FGetTextForItem::CreateLambda([](const ISceneOutlinerTreeItem& Item)->FString
+	{	
+		UContentBundleEngineSubsystem* ContentBundleEngineSubsystem = UContentBundleEngineSubsystem::Get();
+		if (const FActorTreeItem* ActorItem = Item.CastTo<FActorTreeItem>())
 		{
-			if (!RootRow.IsValid())
+			if (AActor* Actor = ActorItem->Actor.Get())
 			{
-				continue;
+				if (const UContentBundleDescriptor * Descriptor = ContentBundleEngineSubsystem->GetContentBundleDescriptor(Actor->GetContentBundleGuid()))
+				{
+					return Descriptor->GetDisplayName();
+				}
 			}
-			
-			TreeViewPtr->SetItemExpansion(RootRow, false);
-			RootRow->SetIsTreeViewItemExpanded(false);
 		}
-	}
-}
-
-void SObjectMixerEditorList::CacheTreeState(const TArray<TWeakPtr<IObjectMixerEditorListFilter>>& InFilterCombo)
-{
-	using LambdaType = void(*)(const TArray<FObjectMixerEditorListRowPtr>&,
-			TArray<FTreeItemStateCache>*,
-			TSharedPtr<STreeView<FObjectMixerEditorListRowPtr>>);
-	
-	static LambdaType RecursivelyCacheTreeState = [](
-		const TArray<FObjectMixerEditorListRowPtr>& InObjects,
-		TArray<FTreeItemStateCache>* TreeItemStateCache,
-		TSharedPtr<STreeView<FObjectMixerEditorListRowPtr>> InTreeViewPtr)
-	{
-		for (const FObjectMixerEditorListRowPtr& TreeViewItem : InObjects)
+		else if (const FActorDescTreeItem* ActorDescItem = Item.CastTo<FActorDescTreeItem>())
 		{
-			const FObjectMixerEditorListRow::FTreeItemUniqueIdentifier& UniqueIdentifier = TreeViewItem->GetUniqueIdentifier();
-
-			if (!UniqueIdentifier.RowName.IsEmpty())
+			if (const FWorldPartitionActorDesc* ActorDesc = ActorDescItem->ActorDescHandle.Get())
 			{
-				TreeItemStateCache->Add(
-					{
-						UniqueIdentifier,
-						InTreeViewPtr->IsItemExpanded(TreeViewItem),
-						InTreeViewPtr->IsItemSelected(TreeViewItem),
-						TreeViewItem->GetVisibilityRules()
-					}
-				);
+				if (const UContentBundleDescriptor* Descriptor = ContentBundleEngineSubsystem->GetContentBundleDescriptor(ActorDesc->GetContentBundleGuid()))
+				{
+					return Descriptor->GetDisplayName();
+				}
 			}
-
-			RecursivelyCacheTreeState(TreeViewItem->GetChildRows(), TreeItemStateCache, InTreeViewPtr);
 		}
-	};
-
-	FFilterComboToStateCaches* Match = Algo::FindByPredicate(
-		FilterComboToStateCaches,
-		[&InFilterCombo](
-		const FFilterComboToStateCaches& Cache)
-		{
-			return Cache.FilterCombo == InFilterCombo;
-		});
 
-	if (Match)
-	{
-		RecursivelyCacheTreeState(TreeViewRootObjects, &Match->Caches, TreeViewPtr);
-	}
-	else
-	{
-		TArray<FTreeItemStateCache> NewCache;
-		RecursivelyCacheTreeState(TreeViewRootObjects, &NewCache, TreeViewPtr);
-		FilterComboToStateCaches.Add({InFilterCombo, MoveTemp(NewCache)});
-	}
-}
+		return TEXT("");
+	});
 
-void SObjectMixerEditorList::RestoreTreeState(const TArray<TWeakPtr<IObjectMixerEditorListFilter>>& InFilterCombo, const bool bFlushCache)
-{
-	using LambdaType = void(*)(const TArray<FObjectMixerEditorListRowPtr>&,
-			TArray<FTreeItemStateCache>*,
-			TSharedPtr<STreeView<FObjectMixerEditorListRowPtr>>,
-			const bool);
-	
-	static LambdaType RecursivelyRestoreTreeState = [](
-		const TArray<FObjectMixerEditorListRowPtr>& InObjects,
-		TArray<FTreeItemStateCache>* TreeItemStateCache,
-		TSharedPtr<STreeView<FObjectMixerEditorListRowPtr>> InTreeViewPtr,
-		const bool bExpandByDefault)
+	FGetTextForItem SubPackageInfoText = FGetTextForItem::CreateLambda([](const ISceneOutlinerTreeItem& Item) -> FString
 	{
-		for (const FObjectMixerEditorListRowPtr& TreeViewItem : InObjects)
+		if (const FActorTreeItem* ActorItem = Item.CastTo<FActorTreeItem>())
 		{
-			const FObjectMixerEditorListRow::FTreeItemUniqueIdentifier& TreeItemUniqueIdentifier = TreeViewItem->GetUniqueIdentifier();
-			
-			if (const FTreeItemStateCache* StateCachePtr = Algo::FindByPredicate(
-				*TreeItemStateCache,
-				[&TreeItemUniqueIdentifier](const FTreeItemStateCache& Other)
+			if (AActor* Actor = ActorItem->Actor.Get())
+			{
+				if (const ILevelInstanceInterface* ActorAsLevelInstance = Cast<ILevelInstanceInterface>(Actor))
 				{
-					
-					return Other.UniqueIdentifier == TreeItemUniqueIdentifier;
+					return ActorAsLevelInstance->GetWorldAssetPackage();
 				}
-			))
-			{
-				InTreeViewPtr->SetItemExpansion(TreeViewItem, StateCachePtr->bIsExpanded);
-				InTreeViewPtr->SetItemSelection(TreeViewItem, StateCachePtr->bIsSelected);
-				
-				TreeViewItem->SetVisibilityRules(StateCachePtr->VisibilityRules);
 			}
-			else
-			{
-				InTreeViewPtr->SetItemExpansion(TreeViewItem, bExpandByDefault);
-			}
-
-			RecursivelyRestoreTreeState(
-				TreeViewItem->GetChildRows(), TreeItemStateCache, InTreeViewPtr,
-				bExpandByDefault
-			);
 		}
-	};
-	
-	bool bExpandByDefault = true;
-	if (const UObjectMixerEditorSettings* Settings = GetDefault<UObjectMixerEditorSettings>())
-	{
-		bExpandByDefault = Settings->bExpandTreeViewItemsByDefault;
-	}
-
-	for (int32 CachesItr = FilterComboToStateCaches.Num() - 1; CachesItr >= 0; CachesItr--)
-	{
-		if (FilterComboToStateCaches[CachesItr].FilterCombo == InFilterCombo)
+		else if (const FActorDescTreeItem* ActorDescItem = Item.CastTo<FActorDescTreeItem>())
 		{
-			RecursivelyRestoreTreeState(
-			  TreeViewRootObjects, &FilterComboToStateCaches[CachesItr].Caches,
-			  TreeViewPtr, bExpandByDefault
-			);
-
-			if (bFlushCache)
+			if (const FWorldPartitionActorDesc* ActorDesc = ActorDescItem->ActorDescHandle.Get())
 			{
-				FilterComboToStateCaches.RemoveAt(CachesItr);
+				if (FName LevelPackage = ActorDesc->GetLevelPackage(); !LevelPackage.IsNone())
+				{
+					return ActorDesc->GetLevelPackage().ToString();
+				}
 			}
-
-			break;
 		}
-	}
-}
 
-void SObjectMixerEditorList::PropagatePropertyChangesToSelectedRows()
-{
-	using LambdaType = void(*)(const TArray<FObjectMixerEditorListRowPtr>&,
-			TSharedPtr<STreeView<FObjectMixerEditorListRowPtr>>,
-			const TArray<FObjectMixerEditorListRow::FPropertyPropagationInfo>&);
-	
-	static LambdaType RecursivelyPropagatePropertyChangesToSelectedRows = [](
-		const TArray<FObjectMixerEditorListRowPtr>& InObjects,
-		TSharedPtr<STreeView<FObjectMixerEditorListRowPtr>> InTreeViewPtr,
-		const TArray<FObjectMixerEditorListRow::FPropertyPropagationInfo>& InPendingPropertyPropagations)
+		return FString();
+	});
+
+	FGetTextForItem SocketInfoText = FGetTextForItem::CreateLambda([](const ISceneOutlinerTreeItem& Item) -> FString
 	{
-		for (const FObjectMixerEditorListRowPtr& TreeViewItem : InObjects)
+		if (const FActorTreeItem* ActorItem = Item.CastTo<FActorTreeItem>())
 		{
-			const FObjectMixerEditorListRow::FTreeItemUniqueIdentifier& TreeItemUniqueIdentifier = TreeViewItem->GetUniqueIdentifier();
+			AActor* Actor = ActorItem->Actor.Get();
 
-				if (const FObjectMixerEditorListRow::FPropertyPropagationInfo* PendingPropagation =
-					Algo::FindByPredicate(
-						InPendingPropertyPropagations,
-						[&TreeItemUniqueIdentifier](const FObjectMixerEditorListRow::FPropertyPropagationInfo& Other)
-						{
-							return Other.RowIdentifier == TreeItemUniqueIdentifier;
-						}))
-				{
-					TreeViewItem->PropagateChangesToSimilarSelectedRowProperties(*PendingPropagation);
-				}
-
-			RecursivelyPropagatePropertyChangesToSelectedRows(
-				TreeViewItem->GetChildRows(), InTreeViewPtr, InPendingPropertyPropagations
-			);
+			if (Actor)
+			{
+				return Actor->GetAttachParentSocketName().ToString();
+			}
 		}
-	};
-
-	if (GetSelectedTreeViewItemCount() > 1)
-	{
-			RecursivelyPropagatePropertyChangesToSelectedRows(
-		   TreeViewRootObjects, TreeViewPtr, PendingPropertyPropagations
-	   );
-	}
-
-	PendingPropertyPropagations.Empty();
-}
-
-void SObjectMixerEditorList::BuildPerformanceCacheAndGenerateHeaderIfNeeded()
-{
-	TSharedPtr<FObjectMixerEditorList> PinnedListModel = ListModelPtr.Pin();
-	check(PinnedListModel);
-	TSharedPtr<FObjectMixerEditorMainPanel> PinnedMainPanelModel = PinnedListModel->GetMainPanelModel().Pin();
-	check(PinnedMainPanelModel);
+		
+		return FString();
+	});
 
-	// If any of the following overrides change, we need to regenerate the header row. Otherwise skip regeneration for performance reasons.	
-	// GetObjectClassesToFilter, GetColumnsToShowByDefault, GetColumnsToExclude,
-	// GetForceAddedColumns, GetObjectMixerPropertyInheritanceInclusionOptions, ShouldIncludeUnsupportedProperties
-	bool bNeedToGenerateHeaders = false;
-	
-	TSet<UClass*> LocalObjectClassesToFilterCache;
-	TSet<FName> LocalColumnsToShowByDefaultCache;
-	TSet<FName> LocalColumnsToExcludeCache;
-	TSet<FName> LocalForceAddedColumnsCache;
-	
-	for (const TObjectPtr<UObjectMixerObjectFilter>& FilterInstance : PinnedMainPanelModel->GetObjectFilterInstances())
+	FGetTextForItem InternalNameInfoText = FGetTextForItem::CreateLambda([](const ISceneOutlinerTreeItem& Item) -> FString
 	{
-		if (!FilterInstance)
+		if (const FActorTreeItem* ActorItem = Item.CastTo<FActorTreeItem>())
 		{
-			UE_LOG(LogObjectMixerEditor, Display, TEXT("%hs: UObjectMixerObjectFilter instance not valid."), __FUNCTION__);
-			continue;
+			if (AActor* Actor = ActorItem->Actor.Get())
+			{
+				return Actor->GetFName().ToString();
+			}
 		}
-
-		LocalObjectClassesToFilterCache.Append(FilterInstance->GetObjectClassesToFilter());
-		LocalColumnsToShowByDefaultCache.Append(FilterInstance->GetColumnsToShowByDefault());
-		LocalColumnsToExcludeCache.Append(FilterInstance->GetColumnsToExclude());
-		LocalForceAddedColumnsCache.Append(FilterInstance->GetForceAddedColumns());
-	}
-
-	if (LocalObjectClassesToFilterCache.Difference(ObjectClassesToFilterCache).Num() > 0 || 
-		ObjectClassesToFilterCache.Difference(LocalObjectClassesToFilterCache).Num() > 0)
-	{
-		ObjectClassesToFilterCache = LocalObjectClassesToFilterCache;
-		bNeedToGenerateHeaders = true;
-	}
-
-	if (LocalColumnsToShowByDefaultCache.Difference(ColumnsToShowByDefaultCache).Num() > 0 || 
-		ColumnsToShowByDefaultCache.Difference(LocalColumnsToShowByDefaultCache).Num() > 0)
-	{
-		ColumnsToShowByDefaultCache = LocalColumnsToShowByDefaultCache;
-		bNeedToGenerateHeaders = true;
-	}
-
-	if (LocalColumnsToExcludeCache.Difference(ColumnsToExcludeCache).Num() > 0 || 
-		ColumnsToExcludeCache.Difference(LocalColumnsToExcludeCache).Num() > 0)
-	{
-		ColumnsToExcludeCache = LocalColumnsToExcludeCache;
-		bNeedToGenerateHeaders = true;
-	}
-
-	if (LocalForceAddedColumnsCache.Difference(ForceAddedColumnsCache).Num() > 0 || 
-		ForceAddedColumnsCache.Difference(LocalForceAddedColumnsCache).Num() > 0)
-	{
-		ForceAddedColumnsCache = LocalForceAddedColumnsCache;
-		bNeedToGenerateHeaders = true;
-	}
-
-	// These properties should be governed by the Main instance
-	if (const UObjectMixerObjectFilter* FilterInstance = PinnedMainPanelModel->GetMainObjectFilterInstance())
-	{
-		if (const EObjectMixerInheritanceInclusionOptions PropertyInheritanceInclusionOptions =
-			FilterInstance->GetObjectMixerPropertyInheritanceInclusionOptions(); 
-			PropertyInheritanceInclusionOptions != PropertyInheritanceInclusionOptionsCache)
+		else if (const FComponentTreeItem* ComponentItem = Item.CastTo<FComponentTreeItem>())
 		{
-			PropertyInheritanceInclusionOptionsCache = PropertyInheritanceInclusionOptions;
-			bNeedToGenerateHeaders = true;
+			if (UActorComponent* Component = ComponentItem->Component.Get())
+			{
+				return Component->GetFName().ToString();
+			}
 		}
-		if (const bool bShouldIncludeUnsupportedProperties = FilterInstance->ShouldIncludeUnsupportedProperties(); 
-			bShouldIncludeUnsupportedProperties != bShouldIncludeUnsupportedPropertiesCache)
+		else if (const FActorDescTreeItem* ActorDescItem = Item.CastTo<FActorDescTreeItem>())
 		{
-			bShouldIncludeUnsupportedPropertiesCache = bShouldIncludeUnsupportedProperties;
-			bNeedToGenerateHeaders = true;
+			if (const FWorldPartitionActorDesc* ActorDesc = ActorDescItem->ActorDescHandle.Get())
+			{
+				return ActorDesc->GetActorName().ToString();
+			}
 		}
-	}
-	else
-	{
-		UE_LOG(LogObjectMixerEditor, Display, TEXT("%hs: No Main UObjectMixerObjectFilter instance found."), __FUNCTION__);
-	}
-	
-	if (bNeedToGenerateHeaders)
-	{
-		GenerateHeaderRow();
-	}
-}
-
-bool DoesValidWorldObjectHaveAcceptableClass(const UObject* Object, const TSet<UClass*>& ObjectClassesToFilterCache)
-{
-	if (IsValid(Object))
-	{
-		for (UClass* Class : ObjectClassesToFilterCache)
+		else if (const FActorFolderTreeItem* ActorFolderItem = Item.CastTo<FActorFolderTreeItem>())
 		{
-			if (Object->IsA(Class))
+			if (const UActorFolder* ActorFolder = Cast<UActorFolder>(ActorFolderItem->GetActorFolder()))
 			{
-				return true;
+				return ActorFolder->GetFName().ToString();
 			}
 		}
-	}
 
-	return false;
-}
+		return FString();
+	});
 
-void CreateOuterRowsForTopLevelRow(
-	UObject* InObject,
-	FObjectMixerEditorListRowPtr& TopLevelRow,
-	const TSet<UObject*>& AllMatchingObjects,
-	TMap<UObject*, FObjectMixerEditorListRowPtr>& ObjectsWithRowCreated,
-	const TSharedRef<SObjectMixerEditorList>& InListView)
-{
-	if (InObject)
+	FGetTextForItem LevelInfoText = FGetTextForItem::CreateLambda([](const ISceneOutlinerTreeItem& Item) -> FString
 	{
-		if (InObject->IsA(ULevel::StaticClass()) ||
-			InObject->IsA(UWorld::StaticClass()) ||
-			InObject->IsA(UPackage::StaticClass()))
-		{
-			return;
-		}
-
-		FObjectMixerEditorListRowPtr OuterRow = nullptr;
-			
-		if (const FObjectMixerEditorListRowPtr* Match = ObjectsWithRowCreated.Find(InObject))
+		if (const FActorTreeItem* ActorItem = Item.CastTo<FActorTreeItem>())
 		{
-			OuterRow = *Match;
-		}
-		else
-		{				
-			OuterRow = MakeShared<FObjectMixerEditorListRow>(
-				InObject,
-				FObjectMixerEditorListRow::MatchingObject,
-				InListView);
-				
-			check(OuterRow.IsValid());
-			ObjectsWithRowCreated.Add(InObject, OuterRow);
+			if (AActor* Actor = ActorItem->Actor.Get())
+			{
+				return FPackageName::GetShortName(Actor->GetPackage()->GetName());
+			}
 		}
-
-		if (TopLevelRow.IsValid())
+		else if (const FActorDescTreeItem* ActorDescItem = Item.CastTo<FActorDescTreeItem>())
 		{
-			OuterRow->AddToChildRows(TopLevelRow);
-				
-			const bool bMatchingContainer = AllMatchingObjects.Contains(InObject);
-				
-			OuterRow->SetRowType(
-				bMatchingContainer ?
-				FObjectMixerEditorListRow::MatchingContainerObject :
-				FObjectMixerEditorListRow::ContainerObject
-			);
+			if (const FWorldPartitionActorDesc* ActorDesc = ActorDescItem->ActorDescHandle.Get())
+			{
+				return FPackageName::GetShortName(ActorDesc->GetActorPackage());
+			}
 		}
-		TopLevelRow = OuterRow;
-			
-		UObject* Outer = nullptr;
-		
-		if (const AActor* AsActor = Cast<AActor>(InObject))
+		else if (const FActorFolderTreeItem* ActorFolderItem = Item.CastTo<FActorFolderTreeItem>())
 		{
-			Outer = AsActor->GetAttachParentActor();
+			if (const UActorFolder* ActorFolder = Cast<UActorFolder>(ActorFolderItem->GetActorFolder()))
+			{
+				return FPackageName::GetShortName(ActorFolder->GetPackage()->GetName());
+			}
 		}
 
-		if (!Outer)
-		{
-			Outer = InObject->GetOuter();
-		}
+		return FString();
+	});
 
-		// Recurse with Outer object
-		if (Outer)
+	FGetTextForItem UncachedLightsInfoText = FGetTextForItem::CreateLambda([](const ISceneOutlinerTreeItem& Item) -> FString
+	{
+		if (const FActorTreeItem* ActorItem = Item.CastTo<FActorTreeItem>())
 		{
-			CreateOuterRowsForTopLevelRow(
-			   Outer, TopLevelRow, AllMatchingObjects,
-			   ObjectsWithRowCreated, InListView
-		   );
+			AActor* Actor = ActorItem->Actor.Get();
+			if (Actor)
+			{
+				return FString::Printf(TEXT("%7d"), Actor->GetNumUncachedStaticLightingInteractions());
+			}
 		}
-	}
-}
+		return FString();
+	});
 
-void CreateFolderRowsForTopLevelRow(
-	FObjectMixerEditorListRowPtr& TopLevelRow,
-	TMap<FName, FObjectMixerEditorListRowPtr>& FolderMap,
-	const TSharedRef<SObjectMixerEditorList>& InListView)
-{
-	if (const TObjectPtr<AActor> AsActor = Cast<AActor>(TopLevelRow->GetObject()))
-	{
-		FFolder BaseActorFolder = AsActor->GetFolder();
+	TSharedPtr<FObjectMixerEditorList> PinnedListModel = ListModelPtr.Pin();
 
-		while (!BaseActorFolder.IsNone())
+	auto AddTextInfoColumn = [this, PinnedListModel, &OutInitOptions](FName ColumnID, TAttribute<FText> ColumnName, FGetTextForItem ColumnInfo)
+	{
+		bool bShouldShowColumn = false;
+		// Load visible columns from SerializedData
+		if (UObjectMixerEditorSerializedData* SerializedData = GetMutableDefault<UObjectMixerEditorSerializedData>())
 		{
-			FObjectMixerEditorListRowPtr FolderRow;
-			
-			if (const FObjectMixerEditorListRowPtr* Match = FolderMap.Find(BaseActorFolder.GetPath()))
+			for (const TSubclassOf<UObjectMixerObjectFilter>& Filter : PinnedListModel->GetObjectFilterClasses())
 			{
-				FolderRow = *Match;
-			}
-			else
-			{
-				FolderRow =
-					MakeShared<FObjectMixerEditorListRow>(
-						BaseActorFolder, FObjectMixerEditorListRow::Folder, InListView,
-						FText::FromName(BaseActorFolder.GetLeafName()));
-
-				FolderMap.Add(BaseActorFolder.GetPath(), FolderRow.ToSharedRef());
+				const FName FilterName = Filter->GetFName();
+				if (SerializedData->IsColumnDataSerialized(FilterName, ColumnID))
+				{
+					// If any filter's serialized data says we should show a column, show it and stomp other filters' saved setting.
+					bShouldShowColumn = SerializedData->ShouldShowColumn(FilterName, ColumnID);
+					break;
+				}
 			}
+		}
+		
+		OutInitOptions.ColumnMap.Add(
+			ColumnID,
+			FSceneOutlinerColumnInfo(
+				bShouldShowColumn ? ESceneOutlinerColumnVisibility::Visible : ESceneOutlinerColumnVisibility::Invisible,
+				20,
+				FCreateSceneOutlinerColumn::CreateLambda([ColumnID, ColumnName, ColumnInfo](ISceneOutliner& Outliner)
+					{
+						return TSharedRef<FObjectMixerOutlinerTextInfoColumn>(
+							MakeShared<FObjectMixerOutlinerTextInfoColumn>(Outliner, ColumnID, ColumnInfo, FText::GetEmpty()));
+					}
+				),
+				true,
+				TOptional<float>(),
+				ColumnName,
+				EHeaderComboVisibility::Never,
+				FOnGetContent::CreateRaw(this, &SObjectMixerEditorList::GenerateHeaderRowContextMenu)
+			)
+		);
 
-			if (FolderRow)
+		HeaderColumnInfos.Add(
 			{
-				FolderRow->AddToChildRows(TopLevelRow);
-									
-				TopLevelRow = FolderRow.ToSharedRef();
+				nullptr, ColumnID, ColumnID,
+				ColumnName.Get(FText::GetEmpty()),
+				EListViewColumnType::BuiltIn,
+				TEXT("Built-In"),
+				false, true, true, false
 			}
+		);
+	};
 
-			BaseActorFolder = BaseActorFolder.GetParent();
+	// The "Level" column should be named "Package Short Name" in wp enabled levels
+	auto LevelColumnName = TAttribute<FText>::CreateLambda([WorldPtr]() -> FText
+	{
+		if (WorldPtr && WorldPtr->IsPartitionedWorld())
+		{
+			return FSceneOutlinerBuiltInColumnTypes::PackageShortName_Localized();
 		}
-	}
+
+		return FSceneOutlinerBuiltInColumnTypes::Level_Localized();
+	});
+
+	AddTextInfoColumn(FSceneOutlinerBuiltInColumnTypes::Level(), LevelColumnName, LevelInfoText);
+	AddTextInfoColumn(FSceneOutlinerBuiltInColumnTypes::Layer(), FSceneOutlinerBuiltInColumnTypes::Layer_Localized(), LayerInfoText);
+	AddTextInfoColumn(FSceneOutlinerBuiltInColumnTypes::DataLayer(), FSceneOutlinerBuiltInColumnTypes::DataLayer_Localized(), DataLayerInfoText);
+	AddTextInfoColumn(FSceneOutlinerBuiltInColumnTypes::ContentBundle(), FSceneOutlinerBuiltInColumnTypes::ContentBundle_Localized(), ContentBundleInfoText);
+	AddTextInfoColumn(FSceneOutlinerBuiltInColumnTypes::SubPackage(), FSceneOutlinerBuiltInColumnTypes::SubPackage_Localized(), SubPackageInfoText);
+	AddTextInfoColumn(FSceneOutlinerBuiltInColumnTypes::Socket(), FSceneOutlinerBuiltInColumnTypes::Socket_Localized(), SocketInfoText);
+	AddTextInfoColumn(FSceneOutlinerBuiltInColumnTypes::IDName(), FSceneOutlinerBuiltInColumnTypes::IDName_Localized(), InternalNameInfoText);
+	AddTextInfoColumn(FSceneOutlinerBuiltInColumnTypes::UncachedLights(), FSceneOutlinerBuiltInColumnTypes::UncachedLights_Localized(), UncachedLightsInfoText);
 }
 
-void SObjectMixerEditorList::GenerateTreeView()
+void SObjectMixerEditorList::SetupColumns(FSceneOutlinerInitializationOptions& OutInitOptions)
 {
-	const TSharedPtr<FObjectMixerEditorList> PinnedListModel = ListModelPtr.Pin();
+	TSharedPtr<FObjectMixerEditorList> PinnedListModel = ListModelPtr.Pin();
 	check(PinnedListModel);
 	
-	if (!ensure(TreeViewPtr.IsValid()))
-	{
-		return;
-	}
-
-	TArray<TWeakPtr<IObjectMixerEditorListFilter>> FilterCombo = {};
+	HeaderColumnInfos.Empty(HeaderColumnInfos.Num());
 
-	if (const TSharedPtr<FObjectMixerEditorMainPanel> PinnedMainPanel = PinnedListModel->GetMainPanelModel().Pin())
+	// Add Built-In Columns to Header
 	{
-		FilterCombo = PinnedMainPanel->GetWeakActiveListFiltersSortedByName();
-	}
-
-	CacheTreeState(FilterCombo);
-	
-	FlushMemory(true);
-
-	BuildPerformanceCacheAndGenerateHeaderIfNeeded();
-
-	const UWorld* EditorWorld = GetWorld();
-	
-	TSet<UObject*> AllMatchingObjects;
-
-	for (const ULevel* Level : EditorWorld->GetLevels())
-	{
-		if (Level && Level->bIsVisible)
+		auto AddBuiltInColumn = [&](
+			ESceneOutlinerColumnVisibility InDefaultVisibility, int32 InPriorityIndex, const FName ColumnID,
+			const FText DisplayText, const bool bCanBeHidden, FCreateSceneOutlinerColumn CreationDelegate = FCreateSceneOutlinerColumn(),
+			TOptional<float> FillSize = TOptional<float>())
 		{
-			ForEachObjectWithOuter(Level,
-			  [this, &AllMatchingObjects](UObject* Object)
-		  {
-			  if (DoesValidWorldObjectHaveAcceptableClass(Object, ObjectClassesToFilterCache))
-			  {
-				  AllMatchingObjects.Add(Object);
-			  }
-		  });
-		}
-	}
-	
-	TMap<FName, FObjectMixerEditorListRowPtr> FolderMap;
-	for (UObject* Object : AllMatchingObjects)
-	{
-		FObjectMixerEditorListRowPtr TopLevelRow = nullptr;
-
-		UObject* NextOuter = Object;
+			bool bShouldShowColumn = InDefaultVisibility == ESceneOutlinerColumnVisibility::Visible ? true : false;
+			
+			// Load visible columns from SerializedData
+			if (UObjectMixerEditorSerializedData* SerializedData = GetMutableDefault<UObjectMixerEditorSerializedData>())
+			{
+				for (const TSubclassOf<UObjectMixerObjectFilter>& Filter : PinnedListModel->GetObjectFilterClasses())
+				{
+					const FName FilterName = Filter->GetFName();
+					if (SerializedData->IsColumnDataSerialized(FilterName, ColumnID))
+					{
+						// If any filter's serialized data says we should show a column, show it and stomp other filters' saved setting.
+						bShouldShowColumn = SerializedData->ShouldShowColumn(FilterName, ColumnID);
+						break;
+					}
+				}
+			}
+			
+			OutInitOptions.ColumnMap.Add(ColumnID, FSceneOutlinerColumnInfo(
+				bShouldShowColumn ? ESceneOutlinerColumnVisibility::Visible : ESceneOutlinerColumnVisibility::Invisible,
+				InPriorityIndex, CreationDelegate, bCanBeHidden,
+				FillSize, DisplayText, EHeaderComboVisibility::Never,
+				FOnGetContent::CreateRaw(this, &SObjectMixerEditorList::GenerateHeaderRowContextMenu)));
 
-		CreateOuterRowsForTopLevelRow(
-			NextOuter, TopLevelRow, AllMatchingObjects,
-			ObjectsToRowsCreated, SharedThis(this)
+			HeaderColumnInfos.Add(
+				{
+					nullptr, ColumnID, ColumnID,
+					DisplayText,
+					EListViewColumnType::BuiltIn,
+					TEXT("Built-In"),
+					bShouldShowColumn, bCanBeHidden, true, false
+				}
+			);
+		};
+
+		AddBuiltInColumn(
+			ESceneOutlinerColumnVisibility::Visible, 2, FObjectMixerOutlinerVisibilityColumn::GetID(),
+			FSceneOutlinerBuiltInColumnTypes::Gutter_Localized(), true,
+			FCreateSceneOutlinerColumn::CreateLambda([](ISceneOutliner& InSceneOutliner)
+			   {
+				   return MakeShared<FObjectMixerOutlinerVisibilityColumn>(InSceneOutliner);
+			   }
+			)
 		);
-		
-		// Now consider folder hierarchy for the top level row's object if desired
-		if (GetTreeViewMode() == EObjectMixerTreeViewMode::Folders)
-		{
-			CreateFolderRowsForTopLevelRow(TopLevelRow, FolderMap, SharedThis(this));
-		}
-
-		TreeViewRootObjects.AddUnique(TopLevelRow);
-	}
 
-	TreeViewRootObjects.StableSort(SortByTypeThenName);
+		AddBuiltInColumn(
+			ESceneOutlinerColumnVisibility::Visible, 3, FObjectMixerOutlinerSoloColumn::GetID(),
+			FObjectMixerOutlinerSoloColumn::GetLocalizedColumnName(), true,
+			FCreateSceneOutlinerColumn::CreateLambda([](ISceneOutliner& InSceneOutliner)
+				{
+					return MakeShared<FObjectMixerOutlinerSoloColumn>(InSceneOutliner);
+				}
+			)
+		);
 
-	if (TSharedPtr<FObjectMixerEditorMainPanel> MainPanel = PinnedListModel->GetMainPanelModel().Pin())
-	{
-		MainPanel->RebuildCollectionSelector();
-	}
+		AddBuiltInColumn(
+			ESceneOutlinerColumnVisibility::Visible, 10, FSceneOutlinerBuiltInColumnTypes::Label(),
+			FSceneOutlinerBuiltInColumnTypes::Label_Localized(), false
+		);
 
-	RefreshList();
+		AddBuiltInColumn(
+			ESceneOutlinerColumnVisibility::Visible, 11, FSceneOutlinerBuiltInColumnTypes::ActorInfo(),
+			FSceneOutlinerBuiltInColumnTypes::ActorInfo_Localized(), true			
+		);
 
-	RestoreTreeState(FilterCombo);
-}
+		AddBuiltInColumn(
+			ESceneOutlinerColumnVisibility::Invisible, 12, FSceneOutlinerBuiltInColumnTypes::SourceControl(),
+			FSceneOutlinerBuiltInColumnTypes::SourceControl_Localized(), true			
+		);
 
-void SObjectMixerEditorList::FindVisibleTreeViewObjects()
-{
-	VisibleTreeViewObjects.Empty();
+		ESceneOutlinerColumnVisibility UnsavedColumnVisibility = ESceneOutlinerColumnVisibility::Invisible;
 
-	for (const FObjectMixerEditorListRowPtr& Row : TreeViewRootObjects)
-	{
-		if (Row->ShouldRowWidgetBeVisible())
+		if (UWorld* WorldPtr = GetWorld())
 		{
-			VisibleTreeViewObjects.Add(Row);
+			CreateActorTextInfoColumns(WorldPtr, OutInitOptions);
+			
+			if (WorldPtr->IsPartitionedWorld())
+			{
+				// We don't want the pinned column in non wp levels
+				AddBuiltInColumn(
+					ESceneOutlinerColumnVisibility::Visible, 0, FSceneOutlinerBuiltInColumnTypes::Pinned(),
+					FSceneOutlinerBuiltInColumnTypes::Pinned_Localized(), true			
+				);
+
+				// We want the unsaved column to be visible by default in partitioned levels
+				UnsavedColumnVisibility = ESceneOutlinerColumnVisibility::Visible;
+			}
 		}
-	}
-}
 
-void SObjectMixerEditorList::FindVisibleObjectsAndRequestTreeRefresh()
-{
-	FindVisibleTreeViewObjects();
-	TreeViewPtr->RequestTreeRefresh();
-}
+		AddBuiltInColumn(
+			UnsavedColumnVisibility, 1, FSceneOutlinerBuiltInColumnTypes::Unsaved(),
+			FSceneOutlinerBuiltInColumnTypes::Unsaved_Localized(), true			
+		);
+	}
 
-void SObjectMixerEditorList::SelectedTreeItemsToSelectedInLevelEditor()
-{
-	if (GEditor && !bShouldPauseSyncSelection)
+	// Add Property Columns to Map
 	{
-		bShouldPauseSyncSelection = true;
-		const TArray<FObjectMixerEditorListRowPtr>& SelectedTreeItems = TreeViewPtr->GetSelectedItems();
-		
-		if (SelectedTreeItems.Num() == 0)
+		TSet<UClass*> SpecifiedClasses;
+		for (const TObjectPtr<UObjectMixerObjectFilter>& Instance : PinnedListModel->GetObjectFilterInstances())
 		{
-			GEditor->SelectNone(true, true, true);
+			SpecifiedClasses.Append(
+				Instance->GetParentAndChildClassesFromSpecifiedClasses(
+					PinnedListModel->ObjectClassesToFilterCache, PinnedListModel->PropertyInheritanceInclusionOptionsCache
+				)
+			);
 		}
-		else
+	
+		for (const UClass* Class : SpecifiedClasses)
 		{
-			TArray<AActor*> ActorsToSelect;
-			for (const FObjectMixerEditorListRowPtr& SelectedRow : SelectedTreeItems)
+			for (TFieldIterator<FProperty> FieldIterator(Class, EFieldIterationFlags::None); FieldIterator; ++FieldIterator)
 			{
-				if (SelectedRow->GetRowType() == FObjectMixerEditorListRow::MatchingObject ||
-					SelectedRow->GetRowType() == FObjectMixerEditorListRow::ContainerObject)
+				if (FProperty* Property = *FieldIterator)
 				{
-					if (AActor* Actor = SelectedRow->GetSelfOrOuterAsActor())
-					{
-						ActorsToSelect.Add(Actor);
-					}
+					AddUniquePropertyColumnInfo(Property, PinnedListModel->bShouldIncludeUnsupportedPropertiesCache, PinnedListModel->ColumnsToExcludeCache);
 				}
 			}
 
-			if (ActorsToSelect.Num())
+			// Check Force Added Columns
+			for (const FName& PropertyName : PinnedListModel->ForceAddedColumnsCache)
 			{
-				FScopedTransaction SyncSelection(LOCTEXT("SyncObjectMixerSelectionToScene","Sync Object Mixer Selection To Scene"));
-				
-				GEditor->SelectNone(true, true, true);
-
-				for (AActor* Actor : ActorsToSelect)
+				if (FProperty* Property = FindFProperty<FProperty>(Class, PropertyName))
 				{
-					Actor->Modify();
-					GEditor->SelectActor( Actor, true, true, true );
+					AddUniquePropertyColumnInfo(Property, true);
 				}
 			}
 		}
 	}
-}
 
-void SObjectMixerEditorList::OnGetRowChildren(FObjectMixerEditorListRowPtr Row, TArray<FObjectMixerEditorListRowPtr>& OutChildren) const
-{
-	if (Row.IsValid())
+	// Alphabetical sort by Property Name
+	HeaderColumnInfos.StableSort([](const FObjectMixerSceneOutlinerColumnInfo& A, const FObjectMixerSceneOutlinerColumnInfo& B)
 	{
-		OutChildren = Row->GetChildRows();
-		
-		if (const int32 HybridIndex = Row->GetOrFindHybridRowIndex(); HybridIndex != INDEX_NONE)
-		{
-			OutChildren.RemoveAt(HybridIndex);
-		}
+		return A.PropertyDisplayText.ToString() < B.PropertyDisplayText.ToString();
+	});
 
-		if (Row->GetShouldExpandAllChildren())
-		{
-			Row->SetShouldExpandAllChildren(false);
-			SetChildExpansionRecursively(Row, true);
-		}
-		else
-		{
-			Row->SetShouldExpandAllChildren(false);
-		}
-	}
-}
+	// Alphabetical sort by Property Category Name
+	HeaderColumnInfos.StableSort([](const FObjectMixerSceneOutlinerColumnInfo& A, const FObjectMixerSceneOutlinerColumnInfo& B)
+	{
+		return A.PropertyCategoryName.LexicalLess(B.PropertyCategoryName);
+	});
 
-void SObjectMixerEditorList::OnRowChildExpansionChange(FObjectMixerEditorListRowPtr Row, const bool bIsExpanded, const bool bIsRecursive) const
-{
-	if (Row.IsValid())
+	// sort by type
+	HeaderColumnInfos.StableSort([](const FObjectMixerSceneOutlinerColumnInfo& A, const FObjectMixerSceneOutlinerColumnInfo& B)
 	{
-		if (bIsRecursive)
+		return A.PropertyType < B.PropertyType;
+	});
+
+	// Add Property Columns to Header
+	{
+		for (const FObjectMixerSceneOutlinerColumnInfo& ColumnInfo : HeaderColumnInfos)
 		{
-			if (bIsExpanded)
+			// Builtin columns are already in Header
+			if (ColumnInfo.PropertyType == EListViewColumnType::BuiltIn)
 			{
-				if (Row->GetRowType() == FObjectMixerEditorListRow::Folder)
-				{
-					Row->SetShouldExpandAllChildren(true);
-				}
+				continue;
 			}
-			else
+			
+			const FText PropertyName = ColumnInfo.PropertyDisplayText;
+
+			if (PropertyName.IsEmpty())
+			{
+				continue;
+			}
+
+			// Figure out if we should show a certain column
+			// First check filter rules
+			bool bShouldShowColumn = PinnedListModel->ColumnsToShowByDefaultCache.Contains(ColumnInfo.ColumnID);
+
+			// Then load visible columns from SerializedData
+			if (UObjectMixerEditorSerializedData* SerializedData = GetMutableDefault<UObjectMixerEditorSerializedData>())
 			{
-				SetChildExpansionRecursively(Row, bIsExpanded);
+				for (const TSubclassOf<UObjectMixerObjectFilter>& Filter : PinnedListModel->GetObjectFilterClasses())
+				{
+					const FName FilterName = Filter->GetFName();
+					if (SerializedData->IsColumnDataSerialized(FilterName, ColumnInfo.ColumnID))
+					{
+						if (SerializedData->ShouldShowColumn(FilterName, ColumnInfo.ColumnID))
+						{
+							// If any filter's serialized data says we should show a column, show it and stomp other filters' saved setting.
+							bShouldShowColumn = true;
+							break;
+						}
+					}
+				}
 			}
+			
+			OutInitOptions.ColumnMap.Add(
+				ColumnInfo.ColumnID,
+				FSceneOutlinerColumnInfo(
+					bShouldShowColumn ? ESceneOutlinerColumnVisibility::Visible : ESceneOutlinerColumnVisibility::Invisible, 30,
+					FCreateSceneOutlinerColumn::CreateLambda([Property = ColumnInfo.PropertyRef](ISceneOutliner& InSceneOutliner)
+					{
+						return MakeShared<FObjectMixerOutlinerPropertyColumn>(InSceneOutliner, Property);
+					}),
+					true,
+					ColumnInfo.FillWidth,
+					PropertyName,
+					EHeaderComboVisibility::Never,
+					FOnGetContent::CreateRaw(this, &SObjectMixerEditorList::GenerateHeaderRowContextMenu)
+				)
+			);
 		}
-		
-		Row->SetIsTreeViewItemExpanded(bIsExpanded);
 	}
 }
 
-void SObjectMixerEditorList::SetChildExpansionRecursively(const FObjectMixerEditorListRowPtr& InRow, const bool bNewIsExpanded) const
+void SObjectMixerEditorList::PropagatePropertyChangesToSelectedRows()
 {
-	if (InRow.IsValid())
+	using LambdaType = void(*)(const TSet<TWeakPtr<ISceneOutlinerTreeItem>>&,
+			TSharedPtr<STreeView<TSharedPtr<ISceneOutlinerTreeItem>>>,
+			const TArray<FObjectMixerEditorListRowData::FPropertyPropagationInfo>&);
+	
+	static LambdaType RecursivelyPropagatePropertyChangesToSelectedRows = [](
+		const TSet<TWeakPtr<ISceneOutlinerTreeItem>>& InObjects,
+		TSharedPtr<STreeView<TSharedPtr<ISceneOutlinerTreeItem>>> InTreeViewPtr,
+		const TArray<FObjectMixerEditorListRowData::FPropertyPropagationInfo>& InPendingPropertyPropagations)
 	{
-		for (const FObjectMixerEditorListRowPtr& Child : InRow->GetChildRows())
+		for (const TWeakPtr<ISceneOutlinerTreeItem>& TreeViewItem : InObjects)
 		{
-			TreeViewPtr->SetItemExpansion(Child, bNewIsExpanded);
-			Child->SetIsTreeViewItemExpanded(bNewIsExpanded);
+			if (TSharedPtr<ISceneOutlinerTreeItem> PinnedItem = TreeViewItem.Pin())
+			{
+				if (FObjectMixerEditorListRowData* RowData = FObjectMixerUtils::GetRowData(PinnedItem))
+				{
+					const FSceneOutlinerTreeItemID TreeItemUniqueIdentifier = PinnedItem->GetID();
+
+					if (const FObjectMixerEditorListRowData::FPropertyPropagationInfo* PendingPropagation =
+						Algo::FindByPredicate(
+							InPendingPropertyPropagations,
+							[&TreeItemUniqueIdentifier](const FObjectMixerEditorListRowData::FPropertyPropagationInfo& Other)
+							{
+								return Other.RowIdentifier == TreeItemUniqueIdentifier;
+							}))
+					{
+						RowData->PropagateChangesToSimilarSelectedRowProperties(PinnedItem.ToSharedRef(), *PendingPropagation);
+					}
+				}
 
-			SetChildExpansionRecursively(Child, bNewIsExpanded);
+				RecursivelyPropagatePropertyChangesToSelectedRows(
+					PinnedItem->GetChildren(), InTreeViewPtr, InPendingPropertyPropagations
+				);
+			}
 		}
+	};
+
+	if (GetSelectedTreeViewItemCount() > 1)
+	{
+		RecursivelyPropagatePropertyChangesToSelectedRows(
+		   GetWeakTreeRootItems(), GetTreeView(), PendingPropertyPropagations
+		);
 	}
+
+	PendingPropertyPropagations.Empty();
 }
 
 #undef LOCTEXT_NAMESPACE
