Unshelved from pending changelist '23770228':

diff --git a/Engine/Plugins/Editor/ObjectMixer/ObjectMixer/Source/ObjectMixer/Private/Views/Widgets/ObjectMixerEditorListMenuContext.cpp b/Engine/Plugins/Editor/ObjectMixer/ObjectMixer/Source/ObjectMixer/Private/Views/Widgets/ObjectMixerEditorListMenuContext.cpp
--- a/Engine/Plugins/Editor/ObjectMixer/ObjectMixer/Source/ObjectMixer/Private/Views/Widgets/ObjectMixerEditorListMenuContext.cpp
+++ b/Engine/Plugins/Editor/ObjectMixer/ObjectMixer/Source/ObjectMixer/Private/Views/Widgets/ObjectMixerEditorListMenuContext.cpp
@@ -2,54 +2,53 @@
 
 #include "Views/Widgets/ObjectMixerEditorListMenuContext.h"
 
+#include "Views/List/ObjectMixerEditorList.h"
+#include "Views/List/ObjectMixerUtils.h"
+#include "Views/List/RowTypes/ObjectMixerEditorListRowFolder.h"
+
 #include "ActorFolderPickingMode.h"
 #include "FolderTreeItem.h"
-#include "Views/MainPanel/ObjectMixerEditorMainPanel.h"
-
-#include "Elements/Framework/TypedElementSelectionSet.h"
-#include "Elements/Interfaces/TypedElementObjectInterface.h"
-#include "Framework/Application/SlateApplication.h"
-#include "Framework/MultiBox/MultiBoxDefs.h"
 #include "LevelEditor.h"
-#include "LevelEditorActions.h"
 #include "LevelEditorContextMenu.h"
 #include "LevelEditorMenuContext.h"
-#include "LevelEditorViewport.h"
 #include "ObjectMixerEditorModule.h"
 #include "Selection.h"
 #include "SSceneOutliner.h"
 #include "ToolMenu.h"
-#include "ToolMenus.h"
 #include "ToolMenuContext.h"
 #include "ToolMenuDelegates.h"
 #include "ToolMenuEntry.h"
+#include "ToolMenus.h"
 #include "ToolMenuSection.h"
+#include "Elements/Framework/TypedElementSelectionSet.h"
+#include "Framework/Application/SlateApplication.h"
 #include "Framework/Commands/GenericCommands.h"
+#include "Framework/MultiBox/MultiBoxDefs.h"
 #include "Widgets/Input/SCheckBox.h"
 #include "Widgets/Input/SEditableTextBox.h"
 #include "Widgets/Layout/SBox.h"
 
-#define LOCTEXT_NAMESPACE "ObjectMixerEditorListMenuContext"
+#define LOCTEXT_NAMESPACE "ObjectMixerEditor"
 
 FName UObjectMixerEditorListMenuContext::DefaultContextBaseMenuName("ObjectMixer.ContextMenuBase");
 
-TSharedPtr<SWidget> UObjectMixerEditorListMenuContext::CreateContextMenu(const FObjectMixerEditorListMenuContextData InData)
+TSharedPtr<SWidget> UObjectMixerEditorListMenuContext::CreateContextMenu(FObjectMixerEditorListMenuContextData InData)
 {
 	if (InData.SelectedItems.Num() == 0)
 	{
 		return nullptr;
 	}
 
-	return BuildContextMenu(InData);
+	return BuildContextMenu(MoveTemp(InData));
 }
 
 bool UObjectMixerEditorListMenuContext::DoesSelectionHaveType(const FObjectMixerEditorListMenuContextData& InData, UClass* Type)
 {
-	if (const TSharedPtr<FObjectMixerEditorListRow>* Match = Algo::FindByPredicate(
+	if (const TSharedPtr<ISceneOutlinerTreeItem>* Match = Algo::FindByPredicate(
 		InData.SelectedItems,
-		[Type](const TSharedPtr<FObjectMixerEditorListRow>& SelectedItem)
+		[Type](const TSharedPtr<ISceneOutlinerTreeItem>& SelectedItem)
 		{
-			const UObject* Object = SelectedItem->GetObject();
+			const UObject* Object = FObjectMixerUtils::GetRowObject(SelectedItem);
 			return Object && Object->IsA(Type);
 		}))
 	{
@@ -109,7 +108,7 @@ void PerformLevelEditorRegistrations(FToolMenuContext& Context)
 	}
 }
 
-TSharedPtr<SWidget> UObjectMixerEditorListMenuContext::BuildContextMenu(const FObjectMixerEditorListMenuContextData& InData)
+TSharedPtr<SWidget> UObjectMixerEditorListMenuContext::BuildContextMenu(const FObjectMixerEditorListMenuContextData InData)
 {	
 	FToolMenuContext Context;
 
@@ -120,7 +119,7 @@ TSharedPtr<SWidget> UObjectMixerEditorListMenuContext::BuildContextMenu(const FO
 	{
 		UObjectMixerEditorListMenuContext* CastContext = CastChecked<UObjectMixerEditorListMenuContext>(InContext);
 		CastContext->Data.SelectedItems.Empty();
-		CastContext->Data.MainPanelPtr.Reset();
+		CastContext->Data.ListModelPtr.Reset();
 	});
 
 	if (DoesSelectionHaveType(InData, AActor::StaticClass()))
@@ -201,25 +200,6 @@ void UObjectMixerEditorListMenuContext::RegisterFoldersOnlyContextMenu()
 		{
 			FToolMenuSection& Section = InMenu->AddSection("Hierarchy", LOCTEXT("HierarchyMenuHeader", "Hierarchy"));
 
-			if (ContextData.SelectedItems.Num() == 1)
-			{
-				const FSlateIcon NewFolderIcon(FAppStyle::GetAppStyleSetName(), "SceneOutliner.NewFolderIcon");
-				
-				Section.AddMenuEntry(
-				   "CreateSubFolder",
-				   LOCTEXT("CreateSubFolder", "Create Sub Folder"),
-				   FText(),
-				   NewFolderIcon,
-				   FUIAction(FExecuteAction::CreateLambda([ContextData]()
-				   {
-				   		if (TSharedPtr<FObjectMixerEditorMainPanel> PinnedPanel = ContextData.MainPanelPtr.Pin())
-				   		{
-				   			PinnedPanel->OnRequestNewFolder(ContextData.SelectedItems[0]->GetFolder());
-				   		}
-				   }))
-			   );
-			}
-
 			Section.AddSubMenu(
 				"SelectSubMenu",
 				LOCTEXT("SelectSubMenu", "Select"),
@@ -364,9 +344,9 @@ void UObjectMixerEditorListMenuContext::CreateSelectCollectionsSubMenu(UToolMenu
 	);
 	Menu->AddMenuEntry("NewCollectionInput", Args);
 						
-	if (const TSharedPtr<FObjectMixerEditorMainPanel> MainPanel = ContextData.MainPanelPtr.Pin())
+	if (const TSharedPtr<FObjectMixerEditorList> List = ContextData.ListModelPtr.Pin())
 	{
-		if (TArray<FName> Collections = MainPanel->GetAllCollectionNames(); Collections.Num() > 0)
+		if (TArray<FName> Collections = List->GetAllCollectionNames(); Collections.Num() > 0)
 		{
 			FToolMenuSection& Section = Menu->FindOrAddSection("Collections");
 			FToolMenuEntry& Separator = Section.AddSeparator("CollectionsSeparator");
@@ -405,17 +385,17 @@ void UObjectMixerEditorListMenuContext::OnFoldersMenuFolderSelected(
 {
 	if (const FFolderTreeItem* FolderTreeItem = Item->CastTo<FFolderTreeItem>())
 	{
-		for (TSharedPtr<FObjectMixerEditorListRow> SelectedItem : ContextData.SelectedItems)
+		for (const TSharedPtr<ISceneOutlinerTreeItem>& SelectedItem : ContextData.SelectedItems)
 		{
-			if (AActor* AsActor = Cast<AActor>(SelectedItem->GetObject()))
+			if (AActor* AsActor = Cast<AActor>(FObjectMixerUtils::GetRowObject(SelectedItem)))
 			{
 				AsActor->SetFolderPath_Recursively(FolderTreeItem->GetFolder().GetPath());
 			}
-			else if (SelectedItem->GetRowType() == FObjectMixerEditorListRow::Folder)
+			else if (FObjectMixerEditorListRowFolder* AsFolder = FObjectMixerUtils::AsFolderRow(SelectedItem))
 			{
-				if (TSharedPtr<FObjectMixerEditorMainPanel> PinnedPanel = ContextData.MainPanelPtr.Pin())
+				if (TSharedPtr<FObjectMixerEditorList> PinnedList = ContextData.ListModelPtr.Pin())
 				{
-					PinnedPanel->OnRequestMoveFolder(SelectedItem->GetFolder(), FolderTreeItem->GetFolder());
+					//	PinnedList->OnRequestMoveFolder(AsFolder->GetFolder(), FolderTreeItem->GetFolder());
 				}
 			}
 		}
@@ -430,40 +410,38 @@ TSharedRef<TSet<FFolder>> UObjectMixerEditorListMenuContext::GatherInvalidMoveTo
 	// We use a pointer here to save copying the whole array for every invocation of the filter delegate
 	TSharedRef<TSet<FFolder>> Exclusions(new TSet<FFolder>());
 
-	for (const TSharedPtr<FObjectMixerEditorListRow>& SelectedItem : ContextData.SelectedItems)
+	for (const TSharedPtr<ISceneOutlinerTreeItem>& SelectedItem : ContextData.SelectedItems)
 	{
-		if (SelectedItem->GetRowType() != FObjectMixerEditorListRow::Folder)
+		if (FObjectMixerEditorListRowFolder* AsFolder = FObjectMixerUtils::AsFolderRow(SelectedItem))
 		{
-			continue;
-		}
+			Exclusions->Add(AsFolder->GetFolder());
 
-		Exclusions->Add(SelectedItem->GetFolder());
-		
-		if (const TSharedPtr<FObjectMixerEditorListRow>& ParentRow = SelectedItem->GetDirectParentRow().Pin())
-		{
-			auto FolderHasOtherSubFolders = [&SelectedItem](const TSharedPtr<FObjectMixerEditorListRow>& InItem)
+			if (const TSharedPtr<ISceneOutlinerTreeItem>& ParentRow = SelectedItem->GetParent())
 			{
-				if (InItem->GetRowType() == FObjectMixerEditorListRow::Folder && InItem != SelectedItem)
+				auto FolderHasOtherSubFolders = [&SelectedItem](const TSharedPtr<ISceneOutlinerTreeItem>& InItem)
 				{
-					return true;
-				}
-				return false;
-			};
+					if (InItem != SelectedItem && FObjectMixerUtils::AsFolderRow(InItem))
+					{
+						return true;
+					}
+					return false;
+				};
 
-			// Exclude this items direct parent if it is a folder and has no other subfolders we can move to
-			bool bFolderHasOtherSubFolders = false;
-			for (const TSharedPtr<FObjectMixerEditorListRow>& ChildRow : ParentRow->GetChildRows())
-			{
-				if (FolderHasOtherSubFolders(ChildRow))
+				// Exclude this items direct parent if it is a folder and has no other subfolders we can move to
+				bool bFolderHasOtherSubFolders = false;
+				for (const TWeakPtr<ISceneOutlinerTreeItem>& ChildRow : ParentRow->GetChildren())
 				{
-					bFolderHasOtherSubFolders = true;
-					break;
+					if (ChildRow.IsValid() && FolderHasOtherSubFolders(ChildRow.Pin()))
+					{
+						bFolderHasOtherSubFolders = true;
+						break;
+					}
+				}
+	
+				if (!bFolderHasOtherSubFolders)
+				{
+					Exclusions->Add(AsFolder->GetFolder());
 				}
-			}
-			
-			if (!bFolderHasOtherSubFolders)
-			{
-				Exclusions->Add(ParentRow->GetFolder());
 			}
 		}
 	}
@@ -481,9 +459,9 @@ void UObjectMixerEditorListMenuContext::FillFoldersSubMenu(UToolMenu* InMenu, FO
 		FSlateIcon(FAppStyle::GetAppStyleSetName(), "SceneOutliner.NewFolderIcon"),
 		FExecuteAction::CreateLambda([ContextData]()
 		{
-			if (const TSharedPtr<FObjectMixerEditorMainPanel> PinnedMainPanel = ContextData.MainPanelPtr.Pin())
+			if (const TSharedPtr<FObjectMixerEditorList> PinnedList = ContextData.ListModelPtr.Pin())
 			{
-				PinnedMainPanel->OnRequestNewFolder();
+				//PinnedList->OnRequestNewFolder();
 			}
 		})
 	);
@@ -551,11 +529,14 @@ void UObjectMixerEditorListMenuContext::OnTextCommitted(const FText& InText, ETe
 void UObjectMixerEditorListMenuContext::SelectDescendentsOfSelectedFolders(
 	FObjectMixerEditorListMenuContextData ContextData, const bool bRecursive)
 {
-	for (const TSharedPtr<FObjectMixerEditorListRow>& SelectedRow : ContextData.SelectedItems)
+	for (const TSharedPtr<ISceneOutlinerTreeItem>& SelectedRow : ContextData.SelectedItems)
 	{
-		if (SelectedRow->GetRowType() == FObjectMixerEditorListRow::Folder)
+		if (FObjectMixerEditorListRowData* RowData = FObjectMixerUtils::GetRowData(SelectedRow))
 		{
-			SelectedRow->SetChildRowsSelected(true, bRecursive);
+			if (FObjectMixerUtils::AsFolderRow(SelectedRow))
+			{
+				// RowData->SetChildRowsSelected(true, bRecursive);
+			}
 		}
 	}
 }
@@ -574,37 +555,37 @@ void UObjectMixerEditorListMenuContext::OnCollectionMenuEntryCheckStateChanged(c
 
 void UObjectMixerEditorListMenuContext::AddObjectsToCollection(const FName Key, const FObjectMixerEditorListMenuContextData ContextData)
 {
-	if (const TSharedPtr<FObjectMixerEditorMainPanel> MainPanel = ContextData.MainPanelPtr.Pin())
+	if (const TSharedPtr<FObjectMixerEditorList> List = ContextData.ListModelPtr.Pin())
 	{
 		TSet<FSoftObjectPath> ObjectPaths;
 
-		for (const TSharedPtr<FObjectMixerEditorListRow>& Item : ContextData.SelectedItems)
+		for (const TSharedPtr<ISceneOutlinerTreeItem>& Item : ContextData.SelectedItems)
 		{
-			if (const UObject* Object = Item->GetObject())
+			if (const UObject* Object = FObjectMixerUtils::GetRowObject(Item))
 			{
 				ObjectPaths.Add(Object);
 			}
 		}
 		
-		MainPanel->RequestAddObjectsToCollection(Key, ObjectPaths);
+		List->RequestAddObjectsToCollection(Key, ObjectPaths);
 	}
 }
 
 void UObjectMixerEditorListMenuContext::RemoveObjectsFromCollection(const FName Key, const FObjectMixerEditorListMenuContextData ContextData)
 {
-	if (const TSharedPtr<FObjectMixerEditorMainPanel> MainPanel = ContextData.MainPanelPtr.Pin())
+	if (const TSharedPtr<FObjectMixerEditorList> List = ContextData.ListModelPtr.Pin())
 	{
 		TSet<FSoftObjectPath> ObjectPaths;
 
-		for (const TSharedPtr<FObjectMixerEditorListRow>& Item : ContextData.SelectedItems)
+		for (const TSharedPtr<ISceneOutlinerTreeItem>& Item : ContextData.SelectedItems)
 		{
-			if (const UObject* Object = Item->GetObject())
+			if (const UObject* Object = FObjectMixerUtils::GetRowObject(Item))
 			{
 				ObjectPaths.Add(Object);
 			}
 		}
 		
-		MainPanel->RequestRemoveObjectsFromCollection(Key, ObjectPaths);
+		List->RequestRemoveObjectsFromCollection(Key, ObjectPaths);
 	}
 }
 
@@ -612,13 +593,13 @@ bool UObjectMixerEditorListMenuContext::AreAllObjectsInCollection(const FName Ke
 {
 	bool bAreAllSelectedObjectsInCollection = false;
 
-	if (const TSharedPtr<FObjectMixerEditorMainPanel> MainPanel = ContextData.MainPanelPtr.Pin())
+	if (const TSharedPtr<FObjectMixerEditorList> List = ContextData.ListModelPtr.Pin())
 	{
-		for (const TSharedPtr<FObjectMixerEditorListRow>& Item : ContextData.SelectedItems)
+		for (const TSharedPtr<ISceneOutlinerTreeItem>& Item : ContextData.SelectedItems)
 		{
-			if (const UObject* Object = Item->GetObject())
+			if (const UObject* Object = FObjectMixerUtils::GetRowObject(Item))
 			{
-				bAreAllSelectedObjectsInCollection = MainPanel->IsObjectInCollection(Key, Object);
+				bAreAllSelectedObjectsInCollection = List->IsObjectInCollection(Key, Object);
 
 				if (!bAreAllSelectedObjectsInCollection)
 				{
@@ -638,13 +619,13 @@ ECheckBoxState UObjectMixerEditorListMenuContext::GetCheckStateForCollection(con
 	int32 ItemsInCollection = 0;
 	int32 ItemsNotInCollection = 0;
 
-	if (const TSharedPtr<FObjectMixerEditorMainPanel> MainPanel = ContextData.MainPanelPtr.Pin())
+	if (const TSharedPtr<FObjectMixerEditorList> List = ContextData.ListModelPtr.Pin())
 	{
-		for (const TSharedPtr<FObjectMixerEditorListRow>& Item : ContextData.SelectedItems)
+		for (const TSharedPtr<ISceneOutlinerTreeItem>& Item : ContextData.SelectedItems)
 		{
-			if (const UObject* Object = Item->GetObject())
+			if (const UObject* Object = FObjectMixerUtils::GetRowObject(Item))
 			{
-				if (MainPanel->IsObjectInCollection(Key, Object))
+				if (List->IsObjectInCollection(Key, Object))
 				{
 					ItemsInCollection++;
 				}
@@ -682,10 +663,10 @@ FToolMenuEntry UObjectMixerEditorListMenuContext::MakeCustomEditMenu(const FObje
 	Entry.MakeCustomWidget.BindLambda(
 		[ContextData](const FToolMenuContext&, const FToolMenuCustomWidgetContext&)
 		{
-			if (TSharedPtr<FObjectMixerEditorMainPanel> PinnedMainPanel = ContextData.MainPanelPtr.Pin())
+			if (TSharedPtr<FObjectMixerEditorList> PinnedList = ContextData.ListModelPtr.Pin())
 			{
 				// Add options with our mapped command list to an FMenuBuilder
-				FMenuBuilder Builder(true, PinnedMainPanel->ObjectMixerElementEditCommands);
+				FMenuBuilder Builder(true, PinnedList->ObjectMixerElementEditCommands);
 				{
 					Builder.AddMenuEntry( FGenericCommands::Get().Cut);
 					Builder.AddMenuEntry( FGenericCommands::Get().Copy);
@@ -707,7 +688,7 @@ FToolMenuEntry UObjectMixerEditorListMenuContext::MakeCustomEditMenu(const FObje
 
 void UObjectMixerEditorListMenuContext::ReplaceEditSubMenu(const FObjectMixerEditorListMenuContextData& ContextData)
 {
-	if (UToolMenu* EditSubMenu = UToolMenus::Get()->ExtendMenu("LevelEditor.ActorContextMenu.EditSubMenu"))
+	if (UToolMenu* EditSubMenu = UToolMenus::Get()->ExtendMenu("LevelEditor.LevelEditorSceneOutliner.ContextMenu.ElementEditActions"))
 	{
 		// Remove existing edit sub menu options
 		FCustomizedToolMenu* MenuCustomization = EditSubMenu->AddMenuCustomization();
diff --git a/Engine/Plugins/Editor/ObjectMixer/ObjectMixer/Source/ObjectMixer/Private/Views/Widgets/SCollectionSelectionButton.cpp b/Engine/Plugins/Editor/ObjectMixer/ObjectMixer/Source/ObjectMixer/Private/Views/Widgets/SCollectionSelectionButton.cpp
--- a/Engine/Plugins/Editor/ObjectMixer/ObjectMixer/Source/ObjectMixer/Private/Views/Widgets/SCollectionSelectionButton.cpp
+++ b/Engine/Plugins/Editor/ObjectMixer/ObjectMixer/Source/ObjectMixer/Private/Views/Widgets/SCollectionSelectionButton.cpp
@@ -2,9 +2,10 @@
 
 #include "Views/Widgets/SCollectionSelectionButton.h"
 
-#include "Views/List/ObjectMixerEditorListFilters/ObjectMixerEditorListFilter_Collection.h"
-#include "Views/MainPanel/ObjectMixerEditorMainPanel.h"
-#include "Views/MainPanel/SObjectMixerEditorMainPanel.h"
+#include "DragAndDrop/ActorDragDropGraphEdOp.h"
+#include "DragAndDrop/ActorDragDropOp.h"
+#include "Views/List/ObjectMixerEditorList.h"
+#include "Views/List/SObjectMixerEditorList.h"
 
 #include "Framework/Commands/GenericCommands.h"
 #include "Framework/MultiBox/MultiBoxBuilder.h"
@@ -43,11 +44,11 @@ public:
 
 void SCollectionSelectionButton::Construct(
 	const FArguments& InArgs,
-	const TSharedRef<SObjectMixerEditorMainPanel> MainPanelWidget,
-	const TSharedRef<FObjectMixerEditorListFilter_Collection> InCollectionListFilter)
+	const TSharedRef<SObjectMixerEditorList> ListWidget,
+	const FName InCollectionName)
 {
-	MainPanelPtr = MainPanelWidget;
-	CollectionListFilter = InCollectionListFilter;
+	ListWidgetPtr = ListWidget;
+	CollectionName = InCollectionName;
 	
 	ChildSlot
 	[
@@ -61,7 +62,7 @@ void SCollectionSelectionButton::Construct(
 			[
 				SAssignNew(EditableTextBlock, SInlineEditableTextBlock)
 				.Style(&FAppStyle::Get().GetWidgetStyle<FInlineEditableTextBlockStyle>("InlineEditableTextBlockSmallStyle"))
-				.Text_Lambda([this, PinnedFilter = CollectionListFilter.Pin()]() { return FText::FromName(PinnedFilter->CollectionName);})
+				.Text(FText::FromName(CollectionName))
 				.Visibility(EVisibility::HitTestInvisible)
 				.ColorAndOpacity(FSlateColor::UseForeground())
 				.OnTextCommitted(this, &SCollectionSelectionButton::OnEditableTextCommitted)
@@ -74,15 +75,13 @@ TSharedRef<SWidget> SCollectionSelectionButton::GetContextMenu() const
 {
 	FMenuBuilder MenuBuilder(true, nullptr);
 
-	const FName CollectionName = CollectionListFilter.Pin()->CollectionName;
-
 	MenuBuilder.AddMenuEntry(
 		FText::Format(LOCTEXT("DeleteCollectionButtonMenuEntry", "Delete '{0}'"), FText::FromName(CollectionName)),
 		LOCTEXT("DeleteCollectionButtonMenuEntryTooltip", "Remove this collection. This operation can be undone."),
 		FGenericCommands::Get().Delete->GetIcon(),
 		FUIAction(FExecuteAction::CreateLambda([this]()
 		{
-			MainPanelPtr.Pin()->RequestRemoveCollection(CollectionListFilter.Pin()->CollectionName);
+			ListWidgetPtr.Pin()->RequestRemoveCollection(CollectionName);
 		}))
 	);
 
@@ -92,9 +91,8 @@ TSharedRef<SWidget> SCollectionSelectionButton::GetContextMenu() const
 		FGenericCommands::Get().Duplicate->GetIcon(),
 		FUIAction(FExecuteAction::CreateLambda([this]()
 		{
-			const TSharedPtr<FObjectMixerEditorListFilter_Collection> CollectionList = CollectionListFilter.Pin();
-			FName DuplicateName = CollectionList->CollectionName;
-			MainPanelPtr.Pin()->RequestDuplicateCollection(CollectionList->CollectionName, DuplicateName);
+			FName DuplicateName = CollectionName;
+			ListWidgetPtr.Pin()->RequestDuplicateCollection(CollectionName, DuplicateName);
 		}))
 	);
 
@@ -116,8 +114,7 @@ void SCollectionSelectionButton::OnEditableTextCommitted(const FText& Text, ETex
 	if (Type == ETextCommit::OnEnter)
 	{
 		const FName NewCollectionName = *Text.ToString();
-		const TSharedPtr<FObjectMixerEditorListFilter_Collection> CollectionFilter = CollectionListFilter.Pin();
-		MainPanelPtr.Pin()->RequestRenameCollection(CollectionFilter->CollectionName, NewCollectionName);
+		ListWidgetPtr.Pin()->RequestRenameCollection(CollectionName, NewCollectionName);
 	}
 }
 
@@ -134,10 +131,9 @@ FReply SCollectionSelectionButton::OnMouseButtonDown(const FGeometry& MyGeometry
 FReply SCollectionSelectionButton::OnMouseButtonUp(const FGeometry& MyGeometry, const FPointerEvent& MouseEvent) 
 {
 	bIsPressed = false;
-	const FName CollectionName = CollectionListFilter.Pin()->CollectionName;
 	if (MouseEvent.GetEffectingButton() == EKeys::LeftMouseButton )
 	{
-		MainPanelPtr.Pin()->OnCollectionCheckedStateChanged(!GetIsChecked(), CollectionName);
+		ListWidgetPtr.Pin()->OnCollectionCheckedStateChanged(!GetIsSelected(), CollectionName);
 	}
 	if (MouseEvent.GetEffectingButton() == EKeys::RightMouseButton)
 	{
@@ -151,7 +147,6 @@ FReply SCollectionSelectionButton::OnMouseButtonUp(const FGeometry& MyGeometry,
 
 FReply SCollectionSelectionButton::OnDragDetected(const FGeometry& MyGeometry, const FPointerEvent& MouseEvent)
 {
-	const FName CollectionName = CollectionListFilter.Pin()->CollectionName;
 	if (CollectionName != UObjectMixerEditorSerializedData::AllCollectionName && CollectionName != NAME_None)
 	{
 		TSharedRef<FCollectionSelectionButtonDragDropOp> OperationFromCollection =
@@ -169,22 +164,29 @@ FReply SCollectionSelectionButton::OnDragDetected(const FGeometry& MyGeometry, c
 
 void SCollectionSelectionButton::OnDragEnter(const FGeometry& MyGeometry, const FDragDropEvent& DragDropEvent)
 {
-	const FName CollectionName = CollectionListFilter.Pin()->CollectionName;
-	
-	if (TSharedPtr<FObjectMixerListRowDragDropOp> OperationFromRow =
-		DragDropEvent.GetOperationAs<FObjectMixerListRowDragDropOp>())
+	if (TSharedPtr<FSceneOutlinerDragDropOp> SceneOutlinerOperation = DragDropEvent.GetOperationAs<FSceneOutlinerDragDropOp>())
 	{
 		if (CollectionName != UObjectMixerEditorSerializedData::AllCollectionName && CollectionName != NAME_None)
 		{
-			OperationFromRow->SetToolTip(
+			SceneOutlinerOperation->SetToolTip(
 			   LOCTEXT("DropRowItemsOntoCollectionButtonCTA","Add selected items to this collection"),
 			   FAppStyle::Get().GetBrush("Graph.ConnectorFeedback.OK"));
 
 			bDropIsValid = true;
 		}
 	}
+	else if (TSharedPtr<FActorDragDropGraphEdOp> ActorOnlyOperation = DragDropEvent.GetOperationAs<FActorDragDropGraphEdOp>())
+	{
+		if (CollectionName != UObjectMixerEditorSerializedData::AllCollectionName && CollectionName != NAME_None)
+		{
+			ActorOnlyOperation->SetToolTip(
+				FActorDragDropGraphEdOp::ToolTip_Compatible, 
+			    LOCTEXT("DropRowActorsOntoCollectionButtonCTA","Add selected actors to this collection"));
 
-	if (TSharedPtr<FCollectionSelectionButtonDragDropOp> OperationFromCollection =
+			bDropIsValid = true;
+		}
+	}
+	else if (TSharedPtr<FCollectionSelectionButtonDragDropOp> OperationFromCollection =
 		DragDropEvent.GetOperationAs<FCollectionSelectionButtonDragDropOp>())
 	{
 		if (CollectionName != UObjectMixerEditorSerializedData::AllCollectionName &&
@@ -206,13 +208,17 @@ void SCollectionSelectionButton::OnDragLeave(const FDragDropEvent& DragDropEvent
 	bIsPressed = false;
 	bDropIsValid = false;
 	
-	if (TSharedPtr<FObjectMixerListRowDragDropOp> OperationFromRow =
-		DragDropEvent.GetOperationAs<FObjectMixerListRowDragDropOp>())
+	if (TSharedPtr<FSceneOutlinerDragDropOp> OperationFromRow =
+		DragDropEvent.GetOperationAs<FSceneOutlinerDragDropOp>())
 	{
 		OperationFromRow->ResetToDefaultToolTip();
 	}
-
-	if (TSharedPtr<FCollectionSelectionButtonDragDropOp> OperationFromCollection =
+	else if (TSharedPtr<FActorDragDropGraphEdOp> ActorOperation =
+		DragDropEvent.GetOperationAs<FActorDragDropGraphEdOp>())
+	{
+		ActorOperation->ResetToDefaultToolTip();
+	}
+	else if (TSharedPtr<FCollectionSelectionButtonDragDropOp> OperationFromCollection =
 		DragDropEvent.GetOperationAs<FCollectionSelectionButtonDragDropOp>())
 	{
 		OperationFromCollection->ResetToDefaultToolTip();
@@ -224,39 +230,57 @@ FReply SCollectionSelectionButton::OnDrop(const FGeometry& MyGeometry, const FDr
 	bIsPressed = false;
 	bDropIsValid = false;
 	
-	const FName CollectionName = CollectionListFilter.Pin()->CollectionName;
-	
-	if (TSharedPtr<FObjectMixerListRowDragDropOp> OperationFromRow =
-		DragDropEvent.GetOperationAs<FObjectMixerListRowDragDropOp>())
+	if (TSharedPtr<FSceneOutlinerDragDropOp> OperationFromRow =
+		DragDropEvent.GetOperationAs<FSceneOutlinerDragDropOp>())
 	{
 		if (CollectionName != UObjectMixerEditorSerializedData::AllCollectionName && CollectionName != NAME_None)
 		{
-			if (const TSharedPtr<FObjectMixerEditorMainPanel> MainPanelModel = MainPanelPtr.Pin()->GetMainPanelModel().Pin())
+			if (const TSharedPtr<FObjectMixerEditorList> ListModel = ListWidgetPtr.Pin()->GetListModelPtr().Pin())
 			{
 				TSet<FSoftObjectPath> ObjectPaths;
 
-				for (const TSharedPtr<FObjectMixerEditorListRow>& Item : OperationFromRow->DraggedItems)
+				for (TWeakObjectPtr<AActor> Item : OperationFromRow->GetSubOp<FActorDragDropOp>()->Actors)
 				{
-					if (const UObject* Object = Item->GetObject())
+					if (Item.IsValid())
 					{
-						ObjectPaths.Add(Object);
+						ObjectPaths.Add(Item.Get());
 					}
 				}
 	
-				MainPanelModel->RequestAddObjectsToCollection(CollectionName, ObjectPaths);
+				ListModel->RequestAddObjectsToCollection(CollectionName, ObjectPaths);
 			}
 		}
 	}
+	else if (TSharedPtr<FActorDragDropGraphEdOp> ActorOperation =
+		DragDropEvent.GetOperationAs<FActorDragDropGraphEdOp>())
+	{
+		if (CollectionName != UObjectMixerEditorSerializedData::AllCollectionName && CollectionName != NAME_None)
+		{
+			if (const TSharedPtr<FObjectMixerEditorList> ListModel = ListWidgetPtr.Pin()->GetListModelPtr().Pin())
+			{
+				TSet<FSoftObjectPath> ObjectPaths;
 
-	if (TSharedPtr<FCollectionSelectionButtonDragDropOp> OperationFromCollection =
+				for (TWeakObjectPtr<AActor> Item : ActorOperation->Actors)
+				{
+					if (Item.IsValid())
+					{
+						ObjectPaths.Add(Item.Get());
+					}
+				}
+	
+				ListModel->RequestAddObjectsToCollection(CollectionName, ObjectPaths);
+			}
+		}
+	}
+	else if (TSharedPtr<FCollectionSelectionButtonDragDropOp> OperationFromCollection =
 		DragDropEvent.GetOperationAs<FCollectionSelectionButtonDragDropOp>())
 	{
 		if (CollectionName != UObjectMixerEditorSerializedData::AllCollectionName &&
 			CollectionName != NAME_None && CollectionName != OperationFromCollection->DraggedItem)
 		{
-			if (const TSharedPtr<FObjectMixerEditorMainPanel> MainPanelModel = MainPanelPtr.Pin()->GetMainPanelModel().Pin())
+			if (const TSharedPtr<FObjectMixerEditorList> ListModel = ListWidgetPtr.Pin()->GetListModelPtr().Pin())
 			{
-				MainPanelModel->RequestReorderCollection(
+				ListModel->RequestReorderCollection(
 					OperationFromCollection->DraggedItem, CollectionName);
 			}
 		}
@@ -267,7 +291,7 @@ FReply SCollectionSelectionButton::OnDrop(const FGeometry& MyGeometry, const FDr
 
 const FSlateBrush* SCollectionSelectionButton::GetBorderBrush() const
 {
-	if (GetIsChecked())
+	if (GetIsSelected())
 	{
 		if (bIsPressed)
 		{
@@ -296,7 +320,7 @@ const FSlateBrush* SCollectionSelectionButton::GetBorderBrush() const
 
 FSlateColor SCollectionSelectionButton::GetBorderForeground() const
 {
-	if (GetIsChecked() || bIsPressed || IsHovered())
+	if (GetIsSelected() || bIsPressed || IsHovered())
 	{
 		return FStyleColors::White;
 	}
@@ -304,9 +328,9 @@ FSlateColor SCollectionSelectionButton::GetBorderForeground() const
 	return FStyleColors::Foreground;
 }
 
-bool SCollectionSelectionButton::GetIsChecked() const
+bool SCollectionSelectionButton::GetIsSelected() const
 {
-	return CollectionListFilter.Pin()->GetIsFilterActive();
+	return ListWidgetPtr.Pin()->IsCollectionSelected(CollectionName);
 }
 
 #undef LOCTEXT_NAMESPACE
diff --git a/Engine/Plugins/Editor/ObjectMixer/ObjectMixer/Source/ObjectMixer/Public/ObjectMixerEditorModule.h b/Engine/Plugins/Editor/ObjectMixer/ObjectMixer/Source/ObjectMixer/Public/ObjectMixerEditorModule.h
--- a/Engine/Plugins/Editor/ObjectMixer/ObjectMixer/Source/ObjectMixer/Public/ObjectMixerEditorModule.h
+++ b/Engine/Plugins/Editor/ObjectMixer/ObjectMixer/Source/ObjectMixer/Public/ObjectMixerEditorModule.h
@@ -6,7 +6,8 @@
 #include "ObjectFilter/ObjectMixerEditorObjectFilter.h"
 #include "Widgets/Docking/SDockTab.h"
 
-class FObjectMixerEditorMainPanel;
+class FObjectMixerEditorList;
+class FObjectMixerEditorList;
 
 class OBJECTMIXEREDITOR_API FObjectMixerEditorModule : public IModuleInterface
 {
@@ -45,8 +46,6 @@ public:
 	/** Called when the Rename command is executed from the UI or hotkey. */
 	virtual void OnRenameCommand();
 	
-	virtual void RequestSyncEditorSelectionToListSelection();
-
 	void RegisterMenuGroup();
 	void UnregisterMenuGroup();
 	virtual void SetupMenuItemVariables();
@@ -68,7 +67,7 @@ public:
 	TSharedPtr<FWorkspaceItem> GetWorkspaceGroup();
 
 	/**
-	 * This is the filter class used to initialize the MainPanel.
+	 * This is the filter class used to initialize the ListModel.
 	 * This filter class cannot be turned off by the end user.
 	 */
 	const TSubclassOf<UObjectMixerObjectFilter>& GetDefaultFilterClass() const;
@@ -83,7 +82,7 @@ protected:
 	TSet<FName> GetPropertiesThatRequireRefresh() const;
 	
 	/** Lives for as long as the module is loaded. */
-	TSharedPtr<FObjectMixerEditorMainPanel> MainPanel;
+	TSharedPtr<FObjectMixerEditorList> ListModel;
 
 	/** The text that appears on the spawned nomad tab */
 	FText TabLabel;
@@ -95,7 +94,7 @@ protected:
 	ETabSpawnerMenuType::Type TabSpawnerType = ETabSpawnerMenuType::Enabled;
 
 	/**
-	 * If set, this is the filter class used to initialize the MainPanel.
+	 * If set, this is the filter class used to initialize the ListModel.
 	 * This filter class cannot be turned off by the end user.
 	 */
 	TSubclassOf<UObjectMixerObjectFilter> DefaultFilterClass;
diff --git a/Engine/Plugins/Editor/ObjectMixer/ObjectMixer/Source/ObjectMixer/Public/Views/List/Modes/ObjectMixerOutlinerMode.h b/Engine/Plugins/Editor/ObjectMixer/ObjectMixer/Source/ObjectMixer/Public/Views/List/Modes/ObjectMixerOutlinerMode.h
new file mode 100644
--- /dev/null
+++ b/Engine/Plugins/Editor/ObjectMixer/ObjectMixer/Source/ObjectMixer/Public/Views/List/Modes/ObjectMixerOutlinerMode.h
@@ -0,0 +1,323 @@
+// Copyright Epic Games, Inc. All Rights Reserved.
+
+#pragma once
+
+#include "ActorMode.h"
+#include "Containers/Map.h"
+#include "EditorConfigBase.h"
+#include "ObjectFilter/ObjectMixerEditorObjectFilter.h"
+#include "UObject/ObjectPtr.h"
+#include "UObject/UObjectGlobals.h"
+
+#include "ObjectMixerOutlinerMode.generated.h"
+
+class FObjectMixerEditorList;
+class IWorldPartitionEditorModule;
+
+namespace ObjectMixerOutliner
+{
+	struct FWeakActorSelectorAcceptingComponents
+	{
+		bool operator()(const TWeakPtr<ISceneOutlinerTreeItem>& Item, TWeakObjectPtr<AActor>& DataOut) const;
+	};
+
+	struct FComponentSelector
+	{
+		bool operator()(const TWeakPtr<ISceneOutlinerTreeItem>& Item, TObjectPtr<UActorComponent>& DataOut) const;
+	};
+}
+
+USTRUCT()
+struct FObjectMixerOutlinerModeConfig
+{
+	GENERATED_BODY()
+
+	/** True when the Scene Outliner is hiding temporary/run-time Actors */
+	UPROPERTY()
+	bool bHideTemporaryActors = false;
+
+	/** True when the Scene Outliner is showing only Actors that exist in the current level */
+	UPROPERTY()
+	bool bShowOnlyActorsInCurrentLevel = false;
+
+	/** True when the Scene Outliner is showing only Actors that exist in the current data layers */
+	UPROPERTY()
+	bool bShowOnlyActorsInCurrentDataLayers = false;
+
+	/** True when the Scene Outliner is showing only Actors that exist in the current content bundle */
+	UPROPERTY()
+	bool bShowOnlyActorsInCurrentContentBundle = false;
+
+	/** True when the Scene Outliner is only displaying selected Actors */
+	UPROPERTY()
+	bool bShowOnlySelectedActors = false;
+
+	/** True when the Scene Outliner is not displaying Actor Components*/
+	UPROPERTY()
+	bool bHideActorComponents = true;
+
+	/** True when the Scene Outliner is not displaying LevelInstances */
+	UPROPERTY()
+	bool bHideLevelInstanceHierarchy = false;
+
+	/** True when the Scene Outliner is not displaying unloaded actors */
+	UPROPERTY()
+	bool bHideUnloadedActors = false;
+
+	/** True when the Scene Outliner is not displaying empty folders */
+	UPROPERTY()
+	bool bHideEmptyFolders = false;
+
+	/** True when the Scene Outliner updates when an actor is selected in the viewport */
+	UPROPERTY()
+	bool bAlwaysFrameSelection = true;
+};
+
+UCLASS(EditorConfig="ObjectMixerOutlinerMode")
+class UObjectMixerOutlinerModeEditorConfig : public UEditorConfigBase
+{
+	GENERATED_BODY()
+	
+public:
+
+	static void Initialize()
+	{
+		if(!Instance)
+		{
+			Instance = NewObject<UObjectMixerOutlinerModeEditorConfig>(); 
+			Instance->AddToRoot();
+		}
+	}
+
+	UPROPERTY(meta=(EditorConfig))
+	TMap<FName, FObjectMixerOutlinerModeConfig> Browsers;
+
+private:
+
+	static TObjectPtr<UObjectMixerOutlinerModeEditorConfig> Instance;
+};
+
+struct OBJECTMIXEREDITOR_API FObjectMixerOutlinerModeParams : FActorModeParams
+{
+	FObjectMixerOutlinerModeParams() {}
+
+	FObjectMixerOutlinerModeParams(
+		SSceneOutliner* InSceneOutliner,
+		const TWeakObjectPtr<UWorld>& InSpecifiedWorldToDisplay = nullptr,
+		bool bInHideComponents = true, bool bInHideLevelInstanceHierarchy = true,
+		bool bInHideUnloadedActors = true, bool bInHideEmptyFolders = true, bool bInHideActorWithNoComponent = true)
+		: FActorModeParams(InSceneOutliner, InSpecifiedWorldToDisplay, bInHideComponents, bInHideLevelInstanceHierarchy, bInHideUnloadedActors, bInHideEmptyFolders)
+	{
+		bHideActorWithNoComponent = bInHideActorWithNoComponent;
+	}
+};
+
+class OBJECTMIXEREDITOR_API FObjectMixerOutlinerMode : public FActorMode
+{
+public:
+	
+	struct FFilterClassSelectionInfo
+	{
+		UClass* Class;
+		bool bIsUserSelected;
+	};
+	
+	FObjectMixerOutlinerMode(const FObjectMixerOutlinerModeParams& Params, const TSharedRef<FObjectMixerEditorList> InListModel);
+	virtual ~FObjectMixerOutlinerMode();
+
+	TWeakPtr<FObjectMixerEditorList> GetListModelPtr() const
+	{
+		return ListModelPtr;
+	}
+
+	SSceneOutliner* GetSceneOutliner() const
+	{
+		return SceneOutliner;
+	}
+	
+	/**
+	 * Determines the style of the tree (flat list or hierarchy)
+	 */
+	EObjectMixerTreeViewMode GetTreeViewMode() const;
+	/**
+	 * Determine the style of the tree (flat list or hierarchy)
+	 */
+	void SetTreeViewMode(EObjectMixerTreeViewMode InViewMode);
+
+	UWorld* GetRepresentingWorld();
+
+	/* Begin ISceneOutlinerMode Interface */
+	virtual bool IsInteractive() const override { return true; }
+	virtual TUniquePtr<ISceneOutlinerHierarchy> CreateHierarchy() override;
+	virtual void Rebuild() override;
+	virtual FCreateSceneOutlinerMode CreateFolderPickerMode(const FFolder::FRootObject& InRootObject = FFolder::GetInvalidRootObject()) const override;
+	virtual void CreateViewContent(FMenuBuilder& MenuBuilder) override;
+	virtual TSharedPtr<FDragDropOperation> CreateDragDropOperation(const FPointerEvent& MouseEvent, const TArray<FSceneOutlinerTreeItemPtr>& InTreeItems) const override;
+	virtual bool ParseDragDrop(FSceneOutlinerDragDropPayload& OutPayload, const FDragDropOperation& Operation) const override;
+	virtual FSceneOutlinerDragValidationInfo ValidateDrop(const ISceneOutlinerTreeItem& DropTarget, const FSceneOutlinerDragDropPayload& Payload) const override;
+	virtual TSharedPtr<SWidget> CreateContextMenu() override;
+	virtual void OnItemAdded(FSceneOutlinerTreeItemPtr Item) override;
+	virtual void OnItemRemoved(FSceneOutlinerTreeItemPtr Item) override;
+	virtual void OnItemSelectionChanged(FSceneOutlinerTreeItemPtr Item, ESelectInfo::Type SelectionType, const FSceneOutlinerItemSelection& Selection) override;
+	virtual void OnItemDoubleClick(FSceneOutlinerTreeItemPtr Item) override;
+	virtual void OnFilterTextCommited(FSceneOutlinerItemSelection& Selection, ETextCommit::Type CommitType) override;
+	virtual void OnItemPassesFilters(const ISceneOutlinerTreeItem& Item) override;
+	virtual FReply OnKeyDown(const FKeyEvent& InKeyEvent) override;
+	virtual void OnDuplicateSelected() override;
+	virtual void OnDrop(ISceneOutlinerTreeItem& DropTarget, const FSceneOutlinerDragDropPayload& Payload, const FSceneOutlinerDragValidationInfo& ValidationInfo) const override;
+	virtual bool CanRenameItem(const ISceneOutlinerTreeItem& Item) const override;
+	virtual FText GetStatusText() const override;
+	virtual FSlateColor GetStatusTextColor() const override;
+	virtual ESelectionMode::Type GetSelectionMode() const override { return ESelectionMode::Multi; }
+	virtual bool SupportsKeyboardFocus() const override { return true; }
+	virtual bool ShouldShowFolders() const override { return GetTreeViewMode() == EObjectMixerTreeViewMode::Folders; }
+	virtual bool SupportsCreateNewFolder() const override { return true; }
+	virtual bool ShowStatusBar() const override { return true; }
+	virtual bool ShowViewButton() const override { return true; }
+	virtual bool ShowFilterOptions() const override { return true; }
+	virtual bool CanDelete() const override;
+	virtual bool CanRename() const override;
+	virtual bool CanCut() const override;
+	virtual bool CanCopy() const override;
+	virtual bool CanPaste() const override;
+	virtual bool CanSupportDragAndDrop() const override { return true; }
+	virtual bool CanCustomizeToolbar() const override { return true; }
+	virtual bool HasErrors() const override;
+	virtual FText GetErrorsText() const override;
+	virtual void RepairErrors() const override;	
+	virtual FFolder CreateNewFolder() override;
+	virtual FFolder GetFolder(const FFolder& ParentPath, const FName& LeafName) override;
+	virtual bool CreateFolder(const FFolder& NewFolder) override;
+	virtual bool ReparentItemToFolder(const FFolder& FolderPath, const FSceneOutlinerTreeItemPtr& Item) override;
+	virtual void SelectFoldersDescendants(const TArray<FFolderTreeItem*>& FolderItems, bool bSelectImmediateChildrenOnly) override;
+	virtual void PinItems(const TArray<FSceneOutlinerTreeItemPtr>& InItems) override;
+	virtual void UnpinItems(const TArray<FSceneOutlinerTreeItemPtr>& InItems) override;
+	virtual bool CanPinItems(const TArray<FSceneOutlinerTreeItemPtr>& InItems) const override;
+	virtual bool CanUnpinItems(const TArray<FSceneOutlinerTreeItemPtr>& InItems) const override;
+	virtual void SynchronizeSelection() override;
+	/* End ISceneOutlinerMode Interface */
+
+	/* External events this mode must respond to */
+
+	/** Called by the engine when a component is updated */
+	void OnComponentsUpdated();
+	/** Called by the engine when an actor is deleted */
+	void OnLevelActorDeleted(AActor* Actor);
+
+	/** Called by the editor to allow selection of unloaded actors */
+	void OnSelectUnloadedActors(const TArray<FGuid>& ActorGuids);
+	
+	/** Called when an actor desc is removed */
+	void OnActorDescRemoved(FWorldPartitionActorDesc* InActorDesc);
+
+	/** Called by engine when edit cut actors begins */
+	void OnEditCutActorsBegin();
+
+	/** Called by engine when edit cut actors ends */
+	void OnEditCutActorsEnd();
+
+	/** Called by engine when edit copy actors begins */
+	void OnEditCopyActorsBegin();
+
+	/** Called by engine when edit copy actors ends */
+	void OnEditCopyActorsEnd();
+
+	/** Called by engine when edit paste actors begins */
+	void OnEditPasteActorsBegin();
+
+	/** Called by engine when edit paste actors ends */
+	void OnEditPasteActorsEnd();
+
+	/** Called by engine when edit duplicate actors begins */
+	void OnDuplicateActorsBegin();
+
+	/** Called by engine when edit duplicate actors ends */
+	void OnDuplicateActorsEnd();
+
+	/** Called by engine when edit delete actors begins */
+	void OnDeleteActorsBegin();
+
+	/** Called by engine when edit delete actors ends */
+	void OnDeleteActorsEnd();
+	
+	/** Function called by the Outliner Filter Bar to compare an item with Type Filters*/
+	virtual bool CompareItemWithClassName(SceneOutliner::FilterBarType InItem, const TSet<FTopLevelAssetPath>&) const override;
+	
+protected:
+
+	/* Events */
+
+	void OnMapChange(uint32 MapFlags);
+	void OnNewCurrentLevel();
+
+	void OnLevelSelectionChanged(UObject* Obj);
+	void OnActorLabelChanged(AActor* ChangedActor);
+	void OnObjectsReplaced(const TMap<UObject*, UObject*>& ReplacementMap);
+	void OnLevelActorRequestsRename(const AActor* Actor);
+	void OnPostLoadMapWithWorld(UWorld* World);
+	
+	/** Build and up the context menu */
+	TSharedPtr<SWidget> BuildContextMenu();
+	/** Register the context menu with the engine */
+	static void RegisterContextMenu();
+	bool CanPasteFoldersOnlyFromClipboard() const;
+
+	bool GetFolderNamesFromPayload(const FSceneOutlinerDragDropPayload& InPayload, TArray<FName>& OutFolders, FFolder::FRootObject& OutCommonRootObject) const;
+	FFolder GetWorldDefaultRootFolder() const;
+
+	void SynchronizeSelectedActorDescs();
+
+	void OnActorEditorContextSubsystemChanged();
+
+	/** Filter factories */
+	static TSharedRef<FSceneOutlinerFilter> CreateShowOnlySelectedActorsFilter();
+	static TSharedRef<FSceneOutlinerFilter> CreateHideTemporaryActorsFilter();
+	static TSharedRef<FSceneOutlinerFilter> CreateIsInCurrentLevelFilter();
+	static TSharedRef<FSceneOutlinerFilter> CreateIsInCurrentDataLayersFilter();
+	static TSharedRef<FSceneOutlinerFilter> CreateHideComponentsFilter();
+	static TSharedRef<FSceneOutlinerFilter> CreateHideLevelInstancesFilter();
+	static TSharedRef<FSceneOutlinerFilter> CreateHideUnloadedActorsFilter();
+	static TSharedRef<FSceneOutlinerFilter> CreateHideEmptyFoldersFilter();
+	TSharedRef<FSceneOutlinerFilter> CreateIsInCurrentContentBundleFilter();
+
+	/** Functions to expose selection framing to the UI */
+	void OnToggleAlwaysFrameSelection();
+	bool ShouldAlwaysFrameSelection();
+
+	/**
+	 * Get a mutable version of the ActorBrowser config for setting values.
+	 * @returns		The config for this ActorBrowser.
+	 * @note		If OutlinerIdentifier is not set for this outliner, it is not possible to store settings.
+	 */
+	FObjectMixerOutlinerModeConfig* GetMutableConfig() const;
+
+	/**
+	 * Get a const version of the ActorBrowser config for getting values.
+	 * @returns		The config for this ActorBrowser.
+	 * @note		If OutlinerIdentifier is not set for this outliner, it is not possible to store settings.
+	 */
+	const FObjectMixerOutlinerModeConfig* GetConstConfig() const;
+
+	/** Save the config for this ActorBrowser */
+	void SaveConfig();
+
+	IWorldPartitionEditorModule* WorldPartitionEditorModule;
+	/** Number of actors (including unloaded) which have passed through the filters */
+	uint32 FilteredActorCount = 0;
+	/** Number of unloaded actors which have passed through all the filters */
+	uint32 FilteredUnloadedActorCount = 0;
+	/** List of unloaded actors which passed through the regular filters and may or may not have passed the search filter */
+	TSet<const FWorldPartitionActorDesc*> ApplicableUnloadedActors;
+	/** List of actors which passed the regular filters and may or may not have passed the search filter */
+	TSet<TWeakObjectPtr<AActor>> ApplicableActors;
+
+	bool bRepresentingWorldGameWorld = false;
+	bool bRepresentingWorldPartitionedWorld = false;
+
+	TWeakPtr<FObjectMixerEditorList> ListModelPtr;
+	
+	TArray<FFilterClassSelectionInfo> FilterClassSelectionInfos;
+	TSharedRef<SWidget> OnGenerateFilterClassMenu();
+
+	bool bShouldPauseSelectionSyncFromEditorTemporarily = false;
+};
diff --git a/Engine/Plugins/Editor/ObjectMixer/ObjectMixer/Source/ObjectMixer/Public/Views/List/ObjectMixerEditorList.h b/Engine/Plugins/Editor/ObjectMixer/ObjectMixer/Source/ObjectMixer/Public/Views/List/ObjectMixerEditorList.h
--- a/Engine/Plugins/Editor/ObjectMixer/ObjectMixer/Source/ObjectMixer/Public/Views/List/ObjectMixerEditorList.h
+++ b/Engine/Plugins/Editor/ObjectMixer/ObjectMixer/Source/ObjectMixer/Public/Views/List/ObjectMixerEditorList.h
@@ -2,61 +2,148 @@
 
 #pragma once
 
-#include "ObjectMixerEditorListRow.h"
-#include "Views/MainPanel/ObjectMixerEditorMainPanel.h"
+#include "ObjectMixerEditorListRowData.h"
 
 #include "Templates/SharedPointer.h"
 #include "Widgets/SWidget.h"
 
+class FObjectMixerEditorModule;
 class SObjectMixerEditorList;
+class UObjectMixerEditorSerializedData;
 
-struct FFolder;
+DECLARE_MULTICAST_DELEGATE(FOnPreFilterChange)
+DECLARE_MULTICAST_DELEGATE(FOnPostFilterChange)
 
-class OBJECTMIXEREDITOR_API FObjectMixerEditorList : public TSharedFromThis<FObjectMixerEditorList>
+
+class OBJECTMIXEREDITOR_API FObjectMixerEditorList : public TSharedFromThis<FObjectMixerEditorList>, public FGCObject
 {
 public:
 
-	FObjectMixerEditorList(TSharedRef<FObjectMixerEditorMainPanel, ESPMode::ThreadSafe> InMainPanel);
+	FObjectMixerEditorList(const FName InModuleName);
 
 	virtual ~FObjectMixerEditorList();
 
-	void FlushWidget();
+	void Initialize();
 	
-	TSharedRef<SWidget> GetOrCreateWidget();
+	void RegisterAndMapContextMenuCommands();
 
-	void OnPreFilterChange();
-	void OnPostFilterChange();
-
-	void ClearList() const;
+	TSharedRef<SWidget> GetOrCreateWidget();
 
-	/**
-	 * Regenerate the list items and refresh the list. Call when adding or removing variables.
-	 */
+	/** Regenerate the list items and refresh the list. Call when adding or removing items. */
 	void RequestRebuildList() const;
 
 	/**
-	 * Refresh filters and sorting.
-	 * Useful for when the list state has gone stale but the item count has not changed.
+	 * Refresh list filters and sorting.
+	 * Useful for when the list state has gone stale but the variable count has not changed.
 	 */
 	void RefreshList() const;
 
+	void BuildPerformanceCache();
+
 	/** Called when the Rename command is executed from the UI or hotkey. */
 	void OnRenameCommand();
+	
+	void RebuildCollectionSelector();
 
-	void RequestSyncEditorSelectionToListSelection() const;
+	void SetDefaultFilterClass(UClass* InNewClass);
+	bool IsClassSelected(UClass* InClass) const;
 
-	[[nodiscard]] TArray<FObjectMixerEditorListRowPtr> GetSelectedTreeViewItems() const;
-	int32 GetSelectedTreeViewItemCount() const;
+	const TArray<TObjectPtr<UObjectMixerObjectFilter>>& GetObjectFilterInstances();
+
+	const UObjectMixerObjectFilter* GetMainObjectFilterInstance();
+
+	void CacheObjectFilterInstances();
+
+	/** Get the style of the tree (flat list or hierarchy) */
+	EObjectMixerTreeViewMode GetTreeViewMode()
+	{
+		return TreeViewMode;
+	}
+	/** Set the style of the tree (flat list or hierarchy)  */
+	void SetTreeViewMode(EObjectMixerTreeViewMode InViewMode)
+	{
+		TreeViewMode = InViewMode;
+		RequestRebuildList();
+	}
+
+	/**
+	 * Force returns result from Filter->ForceGetObjectClassesToFilter.
+	 * Generally, you want to get this from the public performance cache.
+	 */
+	TSet<UClass*> ForceGetObjectClassesToFilter();
+
+	/**
+	 * Force returns result from Filter->ForceGetObjectClassesToPlace.
+	 * Generally, you want to get this from the public performance cache.
+	 */
+	TSet<TSubclassOf<AActor>> ForceGetObjectClassesToPlace();
+
+	const TArray<TSubclassOf<UObjectMixerObjectFilter>>& GetObjectFilterClasses() const
+	{
+		return ObjectFilterClasses;
+	}
+
+	void AddObjectFilterClass(UClass* InObjectFilterClass, const bool bShouldRebuild = true);
+
+	void RemoveObjectFilterClass(UClass* InObjectFilterClass, const bool bCacheAndRebuild = true);
+
+	void ResetObjectFilterClasses(const bool bCacheAndRebuild = true)
+	{
+		ObjectFilterClasses.Empty(ObjectFilterClasses.Num());
+
+		if (bCacheAndRebuild)
+		{
+			CacheAndRebuildFilters();
+		}
+	}
+
+	void CacheAndRebuildFilters()
+	{
+		CacheObjectFilterInstances();
+		RequestRebuildList();
+	}
+
+	/** Used as a way to differentiate different subclasses of the Object Mixer module */
+	FName GetModuleName() const
+	{
+		return ModuleName;
+	}
+
+	// User Collections
+
+	/** Get a pointer to the UObjectMixerEditorSerializedData object along with the name of the filter represented by this ListModel instance. */
+	UObjectMixerEditorSerializedData* GetSerializedData() const;
+	bool RequestAddObjectsToCollection(const FName& CollectionName, const TSet<FSoftObjectPath>& ObjectsToAdd) const;
+	bool RequestRemoveObjectsFromCollection(const FName& CollectionName, const TSet<FSoftObjectPath>& ObjectsToRemove) const;
+	bool RequestRemoveCollection(const FName& CollectionName) const;
+	bool RequestDuplicateCollection(const FName& CollectionToDuplicateName, FName& DesiredDuplicateName) const;
+	bool RequestReorderCollection(const FName& CollectionToMoveName, const FName& CollectionInsertBeforeName) const;
+	bool RequestRenameCollection(const FName& CollectionNameToRename, const FName& NewCollectionName) const;
+	bool DoesCollectionExist(const FName& CollectionName) const;
+	bool IsObjectInCollection(const FName& CollectionName, const FSoftObjectPath& InObject) const;
+	TSet<FName> GetCollectionsForObject(const FSoftObjectPath& InObject) const;
+	TArray<FName> GetAllCollectionNames() const;
+
+	/**
+	 * This is the filter class used to initialize the ListModel.
+	 * This filter class cannot be turned off by the end user.
+	 */
+	const TSubclassOf<UObjectMixerObjectFilter>& GetDefaultFilterClass() const;
+
+	void OnPostFilterChange();
 
-	void OnRequestNewFolder(TOptional<FFolder> ExplicitParentFolder = TOptional<FFolder>());
+	FOnPreFilterChange OnPreFilterChangeDelegate;
+	FOnPostFilterChange OnPostFilterChangeDelegate;
 
-	void OnRequestMoveFolder(const FFolder& FolderToMove, const FFolder& TargetNewParentFolder);
+	TSharedPtr<FUICommandList> ObjectMixerElementEditCommands;
+	TSharedPtr<FUICommandList> ObjectMixerFolderEditCommands;
 
-	void ExecuteListViewSearchOnAllRows(const FString& SearchString, const bool bShouldRefreshAfterward = true);
+	void FlushWidget();
 
-	void EvaluateIfRowsPassFilters(const bool bShouldRefreshAfterward = true) const;
+	UE_NODISCARD TArray<TSharedPtr<ISceneOutlinerTreeItem>> GetSelectedTreeViewItems() const;
+	int32 GetSelectedTreeViewItemCount() const;
 
-	TSet<TWeakPtr<FObjectMixerEditorListRow>> GetSoloRows() const;
+	TSet<TSharedPtr<ISceneOutlinerTreeItem>> GetSoloRows() const;
 	void ClearSoloRows();
 
 	/** Returns true if at least one row is set to Solo. */
@@ -67,12 +154,59 @@ public:
 	 * then sets each object's visibility in editor.
 	 */
 	void EvaluateAndSetEditorVisibilityPerRow();
+	
+	/** Represents the collections the user has selected in the UI. If empty, "All" is considered as selected. */
+	UE_NODISCARD const TSet<FName>& GetSelectedCollections() const;
+	UE_NODISCARD bool IsCollectionSelected(const FName& CollectionName);
+	void SetSelectedCollections(const TSet<FName> InSelectedCollections);
+	void SetCollectionSelected(const FName& CollectionName, const bool bNewSelected);
+
+	// Performance cache
+	TSet<UClass*> ObjectClassesToFilterCache;
+	TSet<FName> ColumnsToShowByDefaultCache;
+	TSet<FName> ColumnsToExcludeCache;
+	TSet<FName> ForceAddedColumnsCache;
+	EObjectMixerInheritanceInclusionOptions PropertyInheritanceInclusionOptionsCache = EObjectMixerInheritanceInclusionOptions::None;
+	bool bShouldIncludeUnsupportedPropertiesCache = false;
+
+protected:
+	
+	virtual void AddReferencedObjects( FReferenceCollector& Collector )  override
+	{
+		Collector.AddReferencedObjects(ObjectFilterInstances);
+	}
+	virtual FString GetReferencerName() const override
+	{
+		return TEXT("FObjectMixerEditorList");
+	}
+
+	/** Represents the collections the user has selected in the UI. If empty, "All" is considered as selected. */
+	TSet<FName> SelectedCollections;
 
-	TWeakPtr<FObjectMixerEditorMainPanel> GetMainPanelModel();
+	TSharedPtr<SObjectMixerEditorList> ListWidget;
+	
+	TArray<TObjectPtr<UObjectMixerObjectFilter>> ObjectFilterInstances;
 
-private:
+	/**
+	 * The classes used to generate property edit columns
+	 */
+	TArray<TSubclassOf<UObjectMixerObjectFilter>> ObjectFilterClasses;
 
-	TWeakPtr<FObjectMixerEditorMainPanel> MainPanelModelPtr;
+	/**
+	 * The FObjectMixerEditorModule subclass that created this instance.
+	 */
+	TWeakPtr<FObjectMixerEditorModule> SpawningModulePtr;
 
-	TSharedPtr<SObjectMixerEditorList> ListWidget;
+	/**
+	 * If set, this is the filter class used to initialize the ListModel.
+	 * This filter class cannot be turned off by the end user.
+	 */
+	TSubclassOf<UObjectMixerObjectFilter> DefaultFilterClass;
+
+	/**
+	 * Determines the style of the tree (flat list or hierarchy)
+	 */
+	EObjectMixerTreeViewMode TreeViewMode = EObjectMixerTreeViewMode::Folders;
+
+	FName ModuleName = NAME_None;
 };
diff --git a/Engine/Plugins/Editor/ObjectMixer/ObjectMixer/Source/ObjectMixer/Public/Views/List/ObjectMixerEditorListRowData.h b/Engine/Plugins/Editor/ObjectMixer/ObjectMixer/Source/ObjectMixer/Public/Views/List/ObjectMixerEditorListRowData.h
new file mode 100644
--- /dev/null
+++ b/Engine/Plugins/Editor/ObjectMixer/ObjectMixer/Source/ObjectMixer/Public/Views/List/ObjectMixerEditorListRowData.h
@@ -0,0 +1,117 @@
+// Copyright Epic Games, Inc. All Rights Reserved.
+
+#pragma once
+
+#include "ObjectFilter/ObjectMixerEditorObjectFilter.h"
+
+#include "ISceneOutlinerTreeItem.h"
+#include "PropertyHandle.h"
+#include "SSceneOutliner.h"
+
+class SObjectMixerEditorList;
+
+struct FSlateBrush;
+
+/** Defines data carried by each row type. */
+struct OBJECTMIXEREDITOR_API FObjectMixerEditorListRowData
+{
+	struct FTransientEditorVisibilityRules
+	{
+		/**
+		 * If true, the user wants the row to be hidden temporarily in the editor.
+		 * This is transient visibility like the eye icon in the Scene Outliner, not the bVisible or bHiddenInGame properties.
+		 */
+		bool bShouldBeHiddenInEditor = false;
+
+		/**
+		 * If true, the user wants the row to have solo visibility. Multiple rows at once can be set to solo.
+		 * Solo rows' objects are exclusively visible,
+		 * so all other objects found in the panel will be invisible while at least one row is in a solo state.
+		 */
+		bool bShouldBeSolo = false;
+	};
+
+	struct FPropertyPropagationInfo
+	{
+		FPropertyPropagationInfo() = default;
+		
+		FSceneOutlinerTreeItemID RowIdentifier;
+		FName PropertyName = NAME_None;
+		EPropertyValueSetFlags::Type PropertyValueSetFlags = 0;
+	};
+	
+	FObjectMixerEditorListRowData(SSceneOutliner* InSceneOutliner, const FText& InDisplayNameOverride = FText::GetEmpty())
+	: SceneOutlinerPtr(InSceneOutliner)
+	, DisplayNameOverride(InDisplayNameOverride)
+	{}
+	
+	FObjectMixerEditorListRowData() = default;
+
+	~FObjectMixerEditorListRowData() = default;
+
+	const TArray<TObjectPtr<UObjectMixerObjectFilter>>& GetObjectFilterInstances() const;
+
+	const UObjectMixerObjectFilter* GetMainObjectFilterInstance() const;
+
+	UE_NODISCARD bool GetIsTreeViewItemExpanded(const TSharedRef<ISceneOutlinerTreeItem> InRow);
+	void SetIsTreeViewItemExpanded(const TSharedRef<ISceneOutlinerTreeItem> InRow, const bool bNewExpanded);
+
+	UE_NODISCARD bool GetDoesRowPassFilters() const;
+	void SetDoesRowPassFilters(const bool bPass);
+
+	UE_NODISCARD bool GetIsSelected(const TSharedRef<ISceneOutlinerTreeItem> InRow);
+	void SetIsSelected(const TSharedRef<ISceneOutlinerTreeItem> InRow, const bool bNewSelected);
+
+	UE_NODISCARD bool HasAtLeastOneChildThatIsNotSolo(
+		const TSharedRef<ISceneOutlinerTreeItem> InRow, const bool bRecursive = true) const;
+
+	UE_NODISCARD FText GetDisplayName(TSharedPtr<ISceneOutlinerTreeItem> InTreeItem, const bool bIsHybridRow = false) const;
+
+	UE_NODISCARD const FText& GetDisplayNameOverride() const
+	{
+		return DisplayNameOverride;
+	}
+
+	void SetDisplayNameOverride(const FText& InDisplayNameOverride)
+	{
+		DisplayNameOverride = InDisplayNameOverride;
+	}
+
+	UE_NODISCARD SObjectMixerEditorList* GetListView() const;
+
+	/**
+	 * Determines the style of the tree (flat list or hierarchy)
+	 */
+	EObjectMixerTreeViewMode GetTreeViewMode();
+
+	UE_NODISCARD TArray<TSharedPtr<ISceneOutlinerTreeItem>> GetSelectedTreeViewItems() const;
+	
+	void OnChangeVisibility(const FSceneOutlinerTreeItemRef TreeItem, const bool bNewVisible);
+	
+	UE_NODISCARD const FTransientEditorVisibilityRules& GetVisibilityRules() const;
+	void SetVisibilityRules(const FTransientEditorVisibilityRules& InVisibilityRules);
+
+	bool IsUserSetHiddenInEditor() const;
+	void SetUserHiddenInEditor(const bool bNewHidden);
+	
+	bool GetRowSoloState() const;
+	void SetRowSoloState(const bool bNewSolo);
+
+	void ClearSoloRows() const;
+
+	bool GetIsItemOrHybridChildSelected(const TSharedRef<ISceneOutlinerTreeItem> InRow);
+
+	void PropagateChangesToSimilarSelectedRowProperties(
+		const TSharedRef<ISceneOutlinerTreeItem> InRow, const FPropertyPropagationInfo PropertyPropagationInfo);
+	
+	TMap<FName, TWeakPtr<IPropertyHandle>> PropertyNamesToHandles;
+
+	SSceneOutliner* SceneOutlinerPtr;
+
+protected:
+	FTransientEditorVisibilityRules VisibilityRules;
+
+	FText DisplayNameOverride;
+	
+	bool bDoesRowPassFilters = true;
+};
diff --git a/Engine/Plugins/Editor/ObjectMixer/ObjectMixer/Source/ObjectMixer/Public/Views/List/RowTypes/ObjectMixerEditorListRowActor.h b/Engine/Plugins/Editor/ObjectMixer/ObjectMixer/Source/ObjectMixer/Public/Views/List/RowTypes/ObjectMixerEditorListRowActor.h
new file mode 100644
--- /dev/null
+++ b/Engine/Plugins/Editor/ObjectMixer/ObjectMixer/Source/ObjectMixer/Public/Views/List/RowTypes/ObjectMixerEditorListRowActor.h
@@ -0,0 +1,34 @@
+// Copyright Epic Games, Inc. All Rights Reserved.
+
+#pragma once
+
+#include "Views/List/ObjectMixerEditorListRowData.h"
+
+#include "ActorTreeItem.h"
+#include "Containers/Array.h"
+#include "Containers/UnrealString.h"
+#include "Folder.h"
+#include "GameFramework/Actor.h"
+#include "Layout/Visibility.h"
+#include "PropertyHandle.h"
+#include "SSceneOutliner.h"
+#include "Templates/SharedPointer.h"
+#include "UObject/Object.h"
+
+struct OBJECTMIXEREDITOR_API FObjectMixerEditorListRowActor : FActorTreeItem
+{
+	explicit FObjectMixerEditorListRowActor(
+		AActor* InObject, 
+		SSceneOutliner* InSceneOutliner, const FText& InDisplayNameOverride = FText::GetEmpty())
+	: FActorTreeItem(InObject)
+	{
+		RowData = FObjectMixerEditorListRowData(InSceneOutliner, InDisplayNameOverride);
+	}
+	
+	FObjectMixerEditorListRowData RowData;
+
+	/* Begin ISceneOutlinerTreeItem Implementation */
+	static const FSceneOutlinerTreeItemType Type;
+	virtual void OnVisibilityChanged(const bool bNewVisibility) override;
+	/* End ISceneOutlinerTreeItem Implementation */
+};
diff --git a/Engine/Plugins/Editor/ObjectMixer/ObjectMixer/Source/ObjectMixer/Public/Views/List/RowTypes/ObjectMixerEditorListRowComponent.h b/Engine/Plugins/Editor/ObjectMixer/ObjectMixer/Source/ObjectMixer/Public/Views/List/RowTypes/ObjectMixerEditorListRowComponent.h
new file mode 100644
--- /dev/null
+++ b/Engine/Plugins/Editor/ObjectMixer/ObjectMixer/Source/ObjectMixer/Public/Views/List/RowTypes/ObjectMixerEditorListRowComponent.h
@@ -0,0 +1,43 @@
+// Copyright Epic Games, Inc. All Rights Reserved.
+
+#pragma once
+
+#include "Views/List/ObjectMixerEditorListRowData.h"
+
+#include "ComponentTreeItem.h"
+#include "Containers/Array.h"
+#include "Containers/UnrealString.h"
+#include "Folder.h"
+#include "GameFramework/Actor.h"
+#include "Layout/Visibility.h"
+#include "PropertyHandle.h"
+#include "SSceneOutliner.h"
+#include "Templates/SharedPointer.h"
+#include "UObject/Object.h"
+
+struct OBJECTMIXEREDITOR_API FObjectMixerEditorListRowComponent : FComponentTreeItem
+{
+	explicit FObjectMixerEditorListRowComponent(
+		UActorComponent* InObject, 
+		SSceneOutliner* InSceneOutliner, const FText& InDisplayNameOverride = FText::GetEmpty())
+	: FComponentTreeItem(InObject)
+	{
+		RowData = FObjectMixerEditorListRowData(InSceneOutliner, InDisplayNameOverride);
+	}
+	
+	FObjectMixerEditorListRowData RowData;
+
+	AActor* GetOwningActor() const
+	{
+		if (const UActorComponent* WrappedComponent = Component.Get())
+		{
+			return WrappedComponent->GetOwner();
+		}
+
+		return nullptr;
+	}
+	
+	/* Begin ISceneOutlinerTreeItem Implementation */
+	static const FSceneOutlinerTreeItemType Type;
+	/* End ISceneOutlinerTreeItem Implementation */
+};
diff --git a/Engine/Plugins/Editor/ObjectMixer/ObjectMixer/Source/ObjectMixer/Public/Views/List/RowTypes/ObjectMixerEditorListRowFolder.h b/Engine/Plugins/Editor/ObjectMixer/ObjectMixer/Source/ObjectMixer/Public/Views/List/RowTypes/ObjectMixerEditorListRowFolder.h
new file mode 100644
--- /dev/null
+++ b/Engine/Plugins/Editor/ObjectMixer/ObjectMixer/Source/ObjectMixer/Public/Views/List/RowTypes/ObjectMixerEditorListRowFolder.h
@@ -0,0 +1,35 @@
+// Copyright Epic Games, Inc. All Rights Reserved.
+
+#pragma once
+
+#include "Views/List/ObjectMixerEditorListRowData.h"
+
+#include "ActorFolderTreeItem.h"
+
+#include "Containers/Array.h"
+#include "Containers/UnrealString.h"
+#include "Folder.h"
+#include "GameFramework/Actor.h"
+#include "Layout/Visibility.h"
+#include "PropertyHandle.h"
+#include "SSceneOutliner.h"
+#include "Templates/SharedPointer.h"
+#include "UObject/Object.h"
+
+struct OBJECTMIXEREDITOR_API FObjectMixerEditorListRowFolder : FActorFolderTreeItem
+{
+	explicit FObjectMixerEditorListRowFolder(
+		const FFolder InFolder, 
+		SSceneOutliner* InSceneOutliner, UWorld* World, const FText& InDisplayNameOverride = FText::GetEmpty())
+	: FActorFolderTreeItem(InFolder, World)
+	{
+		RowData = FObjectMixerEditorListRowData(InSceneOutliner, InDisplayNameOverride);
+	}
+	
+	FObjectMixerEditorListRowData RowData;
+
+	/* Begin ISceneOutlinerTreeItem Implementation */
+	static const FSceneOutlinerTreeItemType Type;
+	virtual void OnVisibilityChanged(const bool bNewVisibility) override;
+	/* End ISceneOutlinerTreeItem Implementation */
+};
diff --git a/Engine/Plugins/Editor/ObjectMixer/ObjectMixer/Source/ObjectMixer/Public/Views/List/RowTypes/ObjectMixerEditorListRowUObject.h b/Engine/Plugins/Editor/ObjectMixer/ObjectMixer/Source/ObjectMixer/Public/Views/List/RowTypes/ObjectMixerEditorListRowUObject.h
new file mode 100644
--- /dev/null
+++ b/Engine/Plugins/Editor/ObjectMixer/ObjectMixer/Source/ObjectMixer/Public/Views/List/RowTypes/ObjectMixerEditorListRowUObject.h
@@ -0,0 +1,44 @@
+// Copyright Epic Games, Inc. All Rights Reserved.
+
+#pragma once
+
+#include "Views/List/ObjectMixerEditorListRowData.h"
+
+#include "Containers/Array.h"
+#include "Containers/UnrealString.h"
+#include "Folder.h"
+#include "GameFramework/Actor.h"
+#include "ISceneOutlinerTreeItem.h"
+#include "Layout/Visibility.h"
+#include "PropertyHandle.h"
+#include "SSceneOutliner.h"
+#include "Templates/SharedPointer.h"
+#include "UObject/Object.h"
+
+struct OBJECTMIXEREDITOR_API FObjectMixerEditorListRowUObject : ISceneOutlinerTreeItem
+{	
+	explicit FObjectMixerEditorListRowUObject(
+		UObject* InObject, 
+		SSceneOutliner* InSceneOutliner, const FText& InDisplayNameOverride = FText::GetEmpty())
+	: ISceneOutlinerTreeItem(Type)
+	, ObjectPtr(InObject)
+	, ID(InObject)
+	{
+		RowData = FObjectMixerEditorListRowData(InSceneOutliner, InDisplayNameOverride);
+	}
+	
+	FObjectMixerEditorListRowData RowData;
+
+	TObjectPtr<UObject> ObjectPtr;
+	
+	/** Constant identifier for this tree item */
+	const FObjectKey ID;
+
+	/* Begin ISceneOutlinerTreeItem Implementation */
+	static const FSceneOutlinerTreeItemType Type;
+	virtual FString GetDisplayString() const override;
+	virtual bool CanInteract() const override { return true; }
+	virtual bool IsValid() const override { return ObjectPtr != nullptr; }
+	virtual FSceneOutlinerTreeItemID GetID() const override { return ID; }
+	/* End ISceneOutlinerTreeItem Implementation */
+};
diff --git a/Engine/Plugins/Editor/ObjectMixer/ObjectMixer/Source/ObjectMixer/Public/Views/List/SObjectMixerEditorList.h b/Engine/Plugins/Editor/ObjectMixer/ObjectMixer/Source/ObjectMixer/Public/Views/List/SObjectMixerEditorList.h
--- a/Engine/Plugins/Editor/ObjectMixer/ObjectMixer/Source/ObjectMixer/Public/Views/List/SObjectMixerEditorList.h
+++ b/Engine/Plugins/Editor/ObjectMixer/ObjectMixer/Source/ObjectMixer/Public/Views/List/SObjectMixerEditorList.h
@@ -3,23 +3,20 @@
 #pragma once
 
 #include "ObjectMixerEditorList.h"
-#include "ObjectMixerEditorListRow.h"
-#include "ObjectMixerEditorListFilters/IObjectMixerEditorListFilter.h"
+#include "ObjectMixerEditorListRowData.h"
+#include "ObjectMixerEditorSerializedData.h"
 
-#include "Widgets/Input/SComboButton.h"
-#include "Widgets/Layout/SGridPanel.h"
-#include "Widgets/SCompoundWidget.h"
-#include "Widgets/Views/SHeaderRow.h"
-#include "Widgets/Views/STreeView.h"
+#include "SSceneOutliner.h"
 
 #include "SObjectMixerEditorList.generated.h"
 
-class SWrapBox;
 class FObjectMixerEditorList;
+class FObjectMixerOutlinerMode;
 class SBox;
 class SComboButton;
 class SSearchBox;
 class SHeaderRow;
+class SWrapBox;
 
 UENUM()
 enum class EListViewColumnType
@@ -29,7 +26,7 @@ enum class EListViewColumnType
 };
 
 USTRUCT()
-struct FListViewColumnInfo
+struct FObjectMixerSceneOutlinerColumnInfo
 {
 	GENERATED_BODY()
 	
@@ -38,6 +35,9 @@ struct FListViewColumnInfo
 	UPROPERTY()
 	FName PropertyName = NAME_None;
 
+	UPROPERTY()
+	FName ColumnID = NAME_None;
+
 	UPROPERTY()
 	FText PropertyDisplayText = FText::GetEmpty();
 
@@ -50,6 +50,9 @@ struct FListViewColumnInfo
 	UPROPERTY()
 	bool bIsDesiredForDisplay = false;
 
+	UPROPERTY()
+	bool bCanBeHidden = true;
+
 	UPROPERTY()
 	bool bCanBeSorted = false;
 
@@ -63,32 +66,56 @@ struct FListViewColumnInfo
 	float FillWidth = 1.0f;
 };
 
-class OBJECTMIXEREDITOR_API SObjectMixerEditorList : public SCompoundWidget
+struct FObjectMixerEditorListRowData;
+
+class OBJECTMIXEREDITOR_API SObjectMixerEditorList : public SSceneOutliner
 {
 
 public:
+	
+	// Columns
+	static const FName ItemNameColumnName;
+	static const FName EditorVisibilityColumnName;
+	static const FName EditorVisibilitySoloColumnName;
 
-	SLATE_BEGIN_ARGS(SObjectMixerEditorList)
-	{}
+	void Construct(const FArguments& InArgs, TSharedRef<FObjectMixerEditorList> ListModel);
 
-	SLATE_END_ARGS()
+	virtual ~SObjectMixerEditorList() {}
 
-	void Construct(const FArguments& InArgs, TSharedRef<FObjectMixerEditorList> ListModel);
+	virtual FReply OnKeyDown( const FGeometry& MyGeometry, const FKeyEvent& InKeyEvent ) override;
+
+	void OnAddDynamicContextMenuSection(UToolMenu* InMenu);
 
-	virtual ~SObjectMixerEditorList() override;
+	/**
+	 * Determines the style of the tree (flat list or hierarchy)
+	 */
+	EObjectMixerTreeViewMode GetTreeViewMode();
+	/**
+	 * Determine the style of the tree (flat list or hierarchy)
+	 */
+	void SetTreeViewMode(EObjectMixerTreeViewMode InViewMode);
+
+	const TSet<FName>& GetSelectedCollections();
+	UE_NODISCARD bool IsCollectionSelected(const FName& CollectionName);
+	
+	void RebuildCollectionSelector();
 
+	bool RequestRemoveCollection(const FName& CollectionName);
+	bool RequestDuplicateCollection(const FName& CollectionToDuplicateName, FName& DesiredDuplicateName) const;
+	bool RequestRenameCollection(const FName& CollectionNameToRename, const FName& NewCollectionName);
+	bool DoesCollectionExist(const FName& CollectionName) const;
+	
+	void OnCollectionCheckedStateChanged(bool bShouldBeChecked, FName CollectionName);
+	ECheckBoxState GetCollectionCheckedState(FName CollectionName) const;
+
+	FObjectMixerOutlinerMode* GetCastedMode() const;
 	UWorld* GetWorld() const;
 
-	TWeakPtr<FObjectMixerEditorList> GetListModelPtr()
+	TWeakPtr<FObjectMixerEditorList> GetListModelPtr() const
 	{
 		return ListModelPtr;
 	}
 
-	void ClearList()
-	{
-		FlushMemory(false);
-	}
-
 	/**
 	 * Refresh filters and sorting.
 	 * Useful for when the list state has gone stale but the variable count has not changed.
@@ -103,36 +130,24 @@ public:
 	/** Called when the Rename command is executed from the UI or hotkey. */
 	void OnRenameCommand();
 
-	void AddToPendingPropertyPropagations(const FObjectMixerEditorListRow::FPropertyPropagationInfo& InPropagationInfo);
+	void AddToPendingPropertyPropagations(const FObjectMixerEditorListRowData::FPropertyPropagationInfo& InPropagationInfo);
 
-	[[nodiscard]] TArray<FObjectMixerEditorListRowPtr> GetSelectedTreeViewItems() const;
+	UE_NODISCARD TArray<TSharedPtr<ISceneOutlinerTreeItem>> GetSelectedTreeViewItems() const;
 	int32 GetSelectedTreeViewItemCount() const;
 
-	void SetTreeViewItemSelected(FObjectMixerEditorListRowPtr Item, const bool bNewSelected);
-	
-	void RequestSyncEditorSelectionToListSelection()
-	{
-		bIsEditorToListSelectionSyncRequested = true;
-	}
-	void SyncEditorSelectionToListSelection();
-
-	void OnRequestNewFolder(TOptional<FFolder> ExplicitParentFolder = TOptional<FFolder>());
-
-	void OnRequestMoveFolder(const FFolder& FolderToMove, const FFolder& TargetNewParentFolder);
-
-	void SetSelectedTreeViewItemActorsEditorVisible(const bool bNewIsVisible, const bool bIsRecursive = false);
+	void SetTreeViewItemSelected(TSharedRef<ISceneOutlinerTreeItem> Item, const bool bNewSelected);
 
-	bool IsTreeViewItemSelected(TSharedRef<FObjectMixerEditorListRow> Item);
+	bool IsTreeViewItemSelected(TSharedPtr<ISceneOutlinerTreeItem> Item);
 
-	[[nodiscard]] TArray<FObjectMixerEditorListRowPtr> GetTreeViewItems() const;
-	void SetTreeViewItems(const TArray<FObjectMixerEditorListRowPtr>& InItems);
-
-	[[nodiscard]] int32 GetTreeViewItemCount() const
+	UE_NODISCARD int32 GetTreeViewItemCount() const
 	{
-		return TreeViewRootObjects.Num();
+		return GetTree().GetItems().Num();
 	}
+
+	TSet<TSharedPtr<ISceneOutlinerTreeItem>> GetTreeRootItems() const;
+	TSet<TWeakPtr<ISceneOutlinerTreeItem>> GetWeakTreeRootItems() const;
 	
-	TSet<TWeakPtr<FObjectMixerEditorListRow>> GetSoloRows() const;
+	TSet<TSharedPtr<ISceneOutlinerTreeItem>> GetSoloRows() const;
 	void ClearSoloRows();
 
 	/** Returns true if at least one row is set to Solo. */
@@ -144,165 +159,71 @@ public:
 	 */
 	void EvaluateAndSetEditorVisibilityPerRow();
 
-	FText GetSearchTextFromSearchInputField() const;
-	FString GetSearchStringFromSearchInputField() const;
-	void ExecuteListViewSearchOnAllRows(const FString& SearchString, const bool bShouldRefreshAfterward = true);
+	bool IsTreeViewItemExpanded(const TSharedPtr<ISceneOutlinerTreeItem>& Row) const;
+	void SetTreeViewItemExpanded(const TSharedPtr<ISceneOutlinerTreeItem>& RowToExpand, const bool bNewExpansion) const;
 
-	bool DoesTreeViewHaveVisibleChildren() const;
+	void PropagatePropertyChangesToSelectedRows();
 
-	bool IsTreeViewItemExpanded(const TSharedPtr<FObjectMixerEditorListRow>& Row) const;
-	void SetTreeViewItemExpanded(const TSharedPtr<FObjectMixerEditorListRow>& RowToExpand, const bool bNewExpansion) const;
+	// Columns
 
-	/**
-	 * Determines the style of the tree (flat list or hierarchy)
-	 */
-	EObjectMixerTreeViewMode GetTreeViewMode();
-	
-	const TArray<TSharedRef<IObjectMixerEditorListFilter>>& GetListFilters();
-	void EvaluateIfRowsPassFilters(const bool bShouldRefreshAfterward = true);
+	FObjectMixerSceneOutlinerColumnInfo* GetColumnInfoByPropertyName(const FName& InPropertyName);
+	TSharedRef<SWidget> GenerateHeaderRowContextMenu() const;
 
-	/** Saves tree item expanded states to be recalled after the tree view is regenerated. */
-	void CacheTreeState(const TArray<TWeakPtr<IObjectMixerEditorListFilter>>& InFilterCombo);
-	void RestoreTreeState(const TArray<TWeakPtr<IObjectMixerEditorListFilter>>& InFilterCombo, const bool bFlushCache = true);
+	/* Begin SSceneOutliner Interface */
+	virtual void CustomAddToToolbar(TSharedPtr<class SHorizontalBox> Toolbar) override;
+	/* End SSceneOutliner Interface */
 
-	void PropagatePropertyChangesToSelectedRows();
+protected:
 
-	// Sorting
+	/** A reference to the struct that controls this widget */
+	TWeakPtr<FObjectMixerEditorList> ListModelPtr;
 
-	const FName& GetActiveSortingColumnName() const
-	{
-		return ActiveSortingColumnName;
-	}
-	EColumnSortMode::Type GetSortModeForColumn(FName InColumnName) const;
-	void OnSortColumnCalled(EColumnSortPriority::Type Priority, const FName& ColumnName, EColumnSortMode::Type SortMode);
-	EColumnSortMode::Type CycleSortMode(const FName& InColumnName);
-	void ExecuteSort(
-		const FName& InColumnName, const EColumnSortMode::Type InColumnSortMode, const bool bShouldRefreshAfterward = true);
-	void ClearSorting()
-	{
-		ActiveSortingColumnName = NAME_None;
-		ActiveSortingType = EColumnSortMode::None;
-	}
+	TSharedPtr<SSearchBox> SearchBoxPtr;
+	TSharedPtr<SComboButton> ViewOptionsComboButton;
 
-	// Columns
+	// User Collections
+	
+	TSharedPtr<SWrapBox> CollectionSelectorBox;
 
-	FListViewColumnInfo* GetColumnInfoByPropertyName(const FName& InPropertyName);
+	bool bIsRebuildRequested = false;
+	
+	TArray<FObjectMixerEditorListRowData::FPropertyPropagationInfo> PendingPropertyPropagations;
 
-	static const FName ItemNameColumnName;
-	static const FName EditorVisibilityColumnName;
-	static const FName EditorVisibilitySoloColumnName;
+	TArray<FObjectMixerSceneOutlinerColumnInfo> HeaderColumnInfos;
 
-	inline static TFunction<bool(const FObjectMixerEditorListRowPtr&, const FObjectMixerEditorListRowPtr&)> SortByTypeThenName =
-		[](const FObjectMixerEditorListRowPtr& A, const FObjectMixerEditorListRowPtr& B)
-	{
-		return (A->GetRowType() < B->GetRowType() ||
-			(A->GetRowType() == B->GetRowType() && A->GetDisplayName().ToString() < B->GetDisplayName().ToString())
-		);
-	};
+	TWeakPtr<SHorizontalBox> ToolbarPtr;
 
-protected:
+	bool CanCreateFolder() const;
+	TSharedRef<SWidget> OnGenerateAddObjectButtonMenu() const;
+	
+	/** Disable all collection filters except CollectionToEnableName */
+	void SetSingleCollectionSelection(const FName& CollectionToEnableName = UObjectMixerEditorSerializedData::AllCollectionName);
 
-	virtual void ExecuteRenameOnPending();
 	virtual void Tick(const FGeometry& AllottedGeometry, const double InCurrentTime, const float InDeltaTime) override;
 
 	void OnActorSpawnedOrDestroyed(AActor* Object)
 	{
 		RequestRebuildList();
 	}
-
-	TWeakPtr<FObjectMixerEditorList> ListModelPtr;
-	
-	TSharedPtr<SHeaderRow> HeaderRow;
-	TSharedRef<SWidget> GenerateHeaderRowContextMenu() const;
-
-	bool bIsRebuildRequested = false;
 	
 	/**
 	 * Regenerate the list items and refresh the list. Call when adding or removing items.
 	 */
 	void RebuildList();
 
+	void InsertCollectionSelector();
+
 	/**
 	 * Only adds properties that pass a series of tests, including having only one unique entry in the column list array.
 	 * @param bForceIncludeProperty If true, only Skiplist and Uniqueness tests will be checked, bypassing class, blueprint editability and other requirements.
 	 * @param PropertySkipList These property names will be skipped when they are encountered in the iteration.
 	 */
-	bool AddUniquePropertyColumnsToHeaderRow(
+	bool AddUniquePropertyColumnInfo(
 		FProperty* Property,
 		const bool bForceIncludeProperty = false,
 		const TSet<FName>& PropertySkipList = {}
 	);
-	void AddBuiltinColumnsToHeaderRow();
-	TSharedPtr<SHeaderRow> GenerateHeaderRow();
-	ECheckBoxState HeaderCheckBoxState = ECheckBoxState::Checked;
-	
-	void FlushMemory(const bool bShouldKeepMemoryAllocated);
-	
-	void SetAllGroupsCollapsed();
-
-	//  Tree View Implementation
-	
-	void BuildPerformanceCacheAndGenerateHeaderIfNeeded();
-	void GenerateTreeView();
-	void FindVisibleTreeViewObjects();
-	void FindVisibleObjectsAndRequestTreeRefresh();
-
-	void SelectedTreeItemsToSelectedInLevelEditor();
-	
-	/** For two-way selection sync, we need to pause selection sync under certain circumstances to prevent infinite loops. */
-	bool bShouldPauseSyncSelection = false;
-	bool bIsEditorToListSelectionSyncRequested = false;
-	
-	void OnGetRowChildren(FObjectMixerEditorListRowPtr Row, TArray<FObjectMixerEditorListRowPtr>& OutChildren) const;
-	void OnRowChildExpansionChange(FObjectMixerEditorListRowPtr Row, const bool bIsExpanded, const bool bIsRecursive = false) const;
-
-	void SetChildExpansionRecursively(const FObjectMixerEditorListRowPtr& InRow, const bool bNewIsExpanded) const;
-
-	TSharedPtr<STreeView<FObjectMixerEditorListRowPtr>> TreeViewPtr;
-
-	struct FTreeItemStateCache
-	{
-		FObjectMixerEditorListRow::FTreeItemUniqueIdentifier UniqueIdentifier;
-		bool bIsExpanded = false;
-		bool bIsSelected = false;
-		FObjectMixerEditorListRow::FTransientEditorVisibilityRules VisibilityRules;
-	};
-
-	struct FFilterComboToStateCaches
-	{
-		TArray<TWeakPtr<IObjectMixerEditorListFilter>> FilterCombo = {};
-		TArray<FTreeItemStateCache> Caches = {};
-	};
-
-	TArray<FFilterComboToStateCaches> FilterComboToStateCaches;
-	
-	FObjectMixerEditorListRow::FTreeItemUniqueIdentifier PendingRenameItem;
-	TArray<FObjectMixerEditorListRow::FPropertyPropagationInfo> PendingPropertyPropagations;
-	
-	TMap<UObject*, FObjectMixerEditorListRowPtr> ObjectsToRowsCreated;
-
-	/** All Tree view objects */
-	TArray<FObjectMixerEditorListRowPtr> TreeViewRootObjects;
-	/** Visible Tree view objects, after filters */
-	TArray<FObjectMixerEditorListRowPtr> VisibleTreeViewObjects;
-
-	TArray<FListViewColumnInfo> ListViewColumns;
-
-	// Sorting
-
-	FName ActiveSortingColumnName = NAME_None;
-	EColumnSortMode::Type ActiveSortingType = EColumnSortMode::None;
-
-	TFunction<bool(const FObjectMixerEditorListRowPtr&, const FObjectMixerEditorListRowPtr&)> SortByOrderAscending = [](const FObjectMixerEditorListRowPtr& A, const FObjectMixerEditorListRowPtr& B)
-		{
-			return A->GetSortOrder() < B->GetSortOrder();
-		};
 
-	// Performance cache
-	TSet<UClass*> ObjectClassesToFilterCache;
-	TSet<FName> ColumnsToShowByDefaultCache;
-	TSet<FName> ColumnsToExcludeCache;
-	TSet<FName> ForceAddedColumnsCache;
-	EObjectMixerInheritanceInclusionOptions PropertyInheritanceInclusionOptionsCache = EObjectMixerInheritanceInclusionOptions::None;
-	bool bShouldIncludeUnsupportedPropertiesCache = false;
+	void CreateActorTextInfoColumns(UWorld *WorldPtr, FSceneOutlinerInitializationOptions& OutInitOptions);
+	void SetupColumns(FSceneOutlinerInitializationOptions& OutInitOptions);
 };
diff --git a/Engine/Plugins/Editor/ObjectMixer/ObjectMixer/Source/ObjectMixer/Public/Views/Widgets/ObjectMixerEditorListMenuContext.h b/Engine/Plugins/Editor/ObjectMixer/ObjectMixer/Source/ObjectMixer/Public/Views/Widgets/ObjectMixerEditorListMenuContext.h
--- a/Engine/Plugins/Editor/ObjectMixer/ObjectMixer/Source/ObjectMixer/Public/Views/Widgets/ObjectMixerEditorListMenuContext.h
+++ b/Engine/Plugins/Editor/ObjectMixer/ObjectMixer/Source/ObjectMixer/Public/Views/Widgets/ObjectMixerEditorListMenuContext.h
@@ -2,7 +2,7 @@
 
 #pragma once
 
-#include "Views/List/ObjectMixerEditorListRow.h"
+#include "Views/List/ObjectMixerEditorListRowData.h"
 
 #include "ToolMenuEntry.h"
 #include "Types/SlateEnums.h"
@@ -21,12 +21,12 @@ public:
 
 	struct FObjectMixerEditorListMenuContextData
 	{
-		TArray<FObjectMixerEditorListRowPtr> SelectedItems;
-		TWeakPtr<class FObjectMixerEditorMainPanel> MainPanelPtr;
+		TArray<TSharedPtr<ISceneOutlinerTreeItem>> SelectedItems;
+		TWeakPtr<class FObjectMixerEditorList> ListModelPtr;
 	};
 
-	static TSharedPtr<SWidget> CreateContextMenu(const FObjectMixerEditorListMenuContextData InData);
-	static TSharedPtr<SWidget> BuildContextMenu(const FObjectMixerEditorListMenuContextData& InData);
+	static TSharedPtr<SWidget> CreateContextMenu(FObjectMixerEditorListMenuContextData InData);
+	static TSharedPtr<SWidget> BuildContextMenu(const FObjectMixerEditorListMenuContextData InData);
 	static void RegisterFoldersOnlyContextMenu();
 	static void RegisterObjectMixerActorContextMenuExtension();
 	static void RegisterObjectMixerElementContextMenuExtension(const FName& MenuName);
diff --git a/Engine/Plugins/Editor/ObjectMixer/ObjectMixer/Source/ObjectMixer/Public/Views/Widgets/SCollectionSelectionButton.h b/Engine/Plugins/Editor/ObjectMixer/ObjectMixer/Source/ObjectMixer/Public/Views/Widgets/SCollectionSelectionButton.h
--- a/Engine/Plugins/Editor/ObjectMixer/ObjectMixer/Source/ObjectMixer/Public/Views/Widgets/SCollectionSelectionButton.h
+++ b/Engine/Plugins/Editor/ObjectMixer/ObjectMixer/Source/ObjectMixer/Public/Views/Widgets/SCollectionSelectionButton.h
@@ -8,10 +8,10 @@
 
 #define LOCTEXT_NAMESPACE "ObjectMixerEditor"
 
+class SObjectMixerEditorList;
 class SInlineEditableTextBlock;
 class SMenuAnchor;
-class SObjectMixerEditorMainPanel;
-class FObjectMixerEditorListFilter_Collection;
+class SObjectMixerEditorList;
 
 class SCollectionSelectionButton final : public SCompoundWidget
 {
@@ -24,8 +24,8 @@ public:
 	
 	void Construct(
 		const FArguments& InArgs,
-		const TSharedRef<SObjectMixerEditorMainPanel> MainPanelWidget,
-		const TSharedRef<FObjectMixerEditorListFilter_Collection> InCollectionListFilter);
+		const TSharedRef<SObjectMixerEditorList> ListWidget,
+		const FName InCollectionName);
 
 	TSharedRef<SWidget> GetContextMenu() const;
 
@@ -46,16 +46,11 @@ public:
 	const FSlateBrush* GetBorderBrush() const;
 	FSlateColor GetBorderForeground() const;
 
-	TWeakPtr<FObjectMixerEditorListFilter_Collection> GetCollectionListFilter()
-	{
-		return  CollectionListFilter;
-	}
+	/** Whether the associated collection is selected */
+	bool GetIsSelected() const;
 
-	/** Whether the associated list filter is active */
-	bool GetIsChecked() const;
-
-	TWeakPtr<SObjectMixerEditorMainPanel> MainPanelPtr;
-	TWeakPtr<FObjectMixerEditorListFilter_Collection> CollectionListFilter;
+	TWeakPtr<SObjectMixerEditorList> ListWidgetPtr;
+	FName CollectionName;
 
 private:
 
