Fix Object Mixer duplicated CVar

diff --git a/Engine/Plugins/Editor/ObjectMixer/ObjectMixer/Source/ObjectMixer/Private/Views/List/Modes/ObjectMixerOutlinerMode.cpp b/Engine/Plugins/Editor/ObjectMixer/ObjectMixer/Source/ObjectMixer/Private/Views/List/Modes/ObjectMixerOutlinerMode.cpp
--- a/Engine/Plugins/Editor/ObjectMixer/ObjectMixer/Source/ObjectMixer/Private/Views/List/Modes/ObjectMixerOutlinerMode.cpp
+++ b/Engine/Plugins/Editor/ObjectMixer/ObjectMixer/Source/ObjectMixer/Private/Views/List/Modes/ObjectMixerOutlinerMode.cpp
@@ -59,10 +59,10 @@
 #include "WorldPartition/WorldPartitionSubsystem.h"
 #include "WorldPartition/DataLayer/DataLayerSubsystem.h"
 
-static int32 GSceneOutlinerAutoRepresentingWorldNetMode = NM_Client;
+static int32 GSceneOutlinerAutoRepresentingWorldNetModeForObjectMixer = NM_Client;
 static FAutoConsoleVariableRef CVarAutoRepresentingWorldNetMode(
-	TEXT("SceneOutliner.AutoRepresentingWorldNetMode"),
-	GSceneOutlinerAutoRepresentingWorldNetMode,
+	TEXT("SceneOutliner.AutoRepresentingWorldNetModeForObjectMixer"),
+	GSceneOutlinerAutoRepresentingWorldNetModeForObjectMixer,
 	TEXT("The preferred NetMode of the world shown in the scene outliner when the 'Auto' option is chosen: 0=Standalone, 1=DedicatedServer, 2=ListenServer, 3=Client"));
 
 #define LOCTEXT_NAMESPACE "ObjectMixerEditor"
@@ -721,7 +721,7 @@ UWorld* FObjectMixerOutlinerMode::GetRepresentingWorld()
 					RepresentingWorld = World;
 					break;
 				}
-				else if ((World->GetNetMode() == ENetMode(GSceneOutlinerAutoRepresentingWorldNetMode)) && (Context.PIEInstance < LowestPIEInstanceSeen))
+				else if ((World->GetNetMode() == ENetMode(GSceneOutlinerAutoRepresentingWorldNetModeForObjectMixer)) && (Context.PIEInstance < LowestPIEInstanceSeen))
 				{
 					RepresentingWorld = World;
 					LowestPIEInstanceSeen = Context.PIEInstance;
