Fix duplicate definitions

diff --git a/Engine/Plugins/Editor/ObjectMixer/ObjectMixer/Source/ObjectMixer/Private/Views/List/Modes/ObjectMixerOutlinerMode.cpp b/Engine/Plugins/Editor/ObjectMixer/ObjectMixer/Source/ObjectMixer/Private/Views/List/Modes/ObjectMixerOutlinerMode.cpp
--- a/Engine/Plugins/Editor/ObjectMixer/ObjectMixer/Source/ObjectMixer/Private/Views/List/Modes/ObjectMixerOutlinerMode.cpp
+++ b/Engine/Plugins/Editor/ObjectMixer/ObjectMixer/Source/ObjectMixer/Private/Views/List/Modes/ObjectMixerOutlinerMode.cpp
@@ -119,10 +119,7 @@ namespace ObjectMixerOutliner
 		}
 		return false;
 	}
-}
 
-namespace SceneOutliner
-{
 	bool FWeakActorSelector::operator()(const TWeakPtr<ISceneOutlinerTreeItem>& Item, TWeakObjectPtr<AActor>& DataOut) const
 	{
 		if (TSharedPtr<ISceneOutlinerTreeItem> ItemPtr = Item.Pin())
@@ -1681,7 +1678,7 @@ void FObjectMixerOutlinerMode::OnItemSelectionChanged(FSceneOutlinerTreeItemPtr
 	if (GetDefault<UObjectMixerEditorSettings>()->bSyncSelection)
 	{
 		bShouldPauseSelectionSyncFromEditorTemporarily = true;
-		TArray<AActor*> SelectedActors = Selection.GetData<AActor*>(SceneOutliner::FActorSelector());
+		TArray<AActor*> SelectedActors = Selection.GetData<AActor*>(ObjectMixerOutliner::FActorSelector());
 		
 		SynchronizeSelectedActorDescs();
 
@@ -1787,7 +1784,7 @@ void FObjectMixerOutlinerMode::OnItemDoubleClick(FSceneOutlinerTreeItemPtr Item)
 			if (Selection.Has<FActorTreeItem>())
 			{
 				const bool bActiveViewportOnly = false;
-				GEditor->MoveViewportCamerasToActor(Selection.GetData<AActor*>(SceneOutliner::FActorSelector()), bActiveViewportOnly);
+				GEditor->MoveViewportCamerasToActor(Selection.GetData<AActor*>(ObjectMixerOutliner::FActorSelector()), bActiveViewportOnly);
 			}
 		}
 		else
@@ -1812,7 +1809,7 @@ void FObjectMixerOutlinerMode::OnFilterTextCommited(FSceneOutlinerItemSelection&
 	const bool bDeselectBSPSurfs = false;
 	const bool WarnAboutManyActors = true;
 	GEditor->SelectNone(bNoteSelectionChange, bDeselectBSPSurfs, WarnAboutManyActors);
-	for (AActor* Actor : Selection.GetData<AActor*>(SceneOutliner::FActorSelector()))
+	for (AActor* Actor : Selection.GetData<AActor*>(ObjectMixerOutliner::FActorSelector()))
 	{
 		const bool bShouldSelect = true;
 		const bool bSelectEvenIfHidden = false;
@@ -2081,7 +2078,7 @@ bool FObjectMixerOutlinerMode::CanPasteFoldersOnlyFromClipboard() const
 
 bool FObjectMixerOutlinerMode::GetFolderNamesFromPayload(const FSceneOutlinerDragDropPayload& InPayload, TArray<FName>& OutFolders, FFolder::FRootObject& OutCommonRootObject) const
 {
-	return FFolder::GetFolderPathsAndCommonRootObject(InPayload.GetData<FFolder>(SceneOutliner::FFolderPathSelector()), OutFolders, OutCommonRootObject);
+	return FFolder::GetFolderPathsAndCommonRootObject(InPayload.GetData<FFolder>(ObjectMixerOutliner::FFolderPathSelector()), OutFolders, OutCommonRootObject);
 }
 
 TSharedPtr<FDragDropOperation> FObjectMixerOutlinerMode::CreateDragDropOperation(const FPointerEvent& MouseEvent, const TArray<FSceneOutlinerTreeItemPtr>& InTreeItems) const
@@ -2091,7 +2088,7 @@ TSharedPtr<FDragDropOperation> FObjectMixerOutlinerMode::CreateDragDropOperation
 	// If the drag contains only actors, we shortcut and create a simple FActorDragDropGraphEdOp rather than an FSceneOutlinerDragDrop composite op.
 	if (DraggedObjects.Has<FActorTreeItem>() && !DraggedObjects.Has<FComponentTreeItem>() && !DraggedObjects.Has<FFolderTreeItem>())
 	{
-		return FActorDragDropGraphEdOp::New(DraggedObjects.GetData<TWeakObjectPtr<AActor>>(SceneOutliner::FWeakActorSelector()));
+		return FActorDragDropGraphEdOp::New(DraggedObjects.GetData<TWeakObjectPtr<AActor>>(ObjectMixerOutliner::FWeakActorSelector()));
 	}
 
 	TSharedPtr<FSceneOutlinerDragDropOp> OutlinerOp = MakeShareable(new FSceneOutlinerDragDropOp());
@@ -2161,7 +2158,7 @@ void FObjectMixerOutlinerMode::SynchronizeSelectedActorDescs()
 	if (UWorldPartitionSubsystem* WorldPartitionSubsystem = UWorld::GetSubsystem<UWorldPartitionSubsystem>(RepresentingWorld.Get()))
 	{
 		const FSceneOutlinerItemSelection Selection = SceneOutliner->GetSelection();
-		TArray<FWorldPartitionActorDesc*> SelectedActorDescs = Selection.GetData<FWorldPartitionActorDesc*>(SceneOutliner::FActorDescSelector());
+		TArray<FWorldPartitionActorDesc*> SelectedActorDescs = Selection.GetData<FWorldPartitionActorDesc*>(ObjectMixerOutliner::FActorDescSelector());
 
 		WorldPartitionSubsystem->SelectedActorDescs.Empty();
 		for (FWorldPartitionActorDesc* SelectedActorDesc : SelectedActorDescs)
@@ -2223,7 +2220,7 @@ FSceneOutlinerDragValidationInfo FObjectMixerOutlinerMode::ValidateDrop(const IS
 		FText AttachErrorMsg;
 		bool bCanAttach = true;
 		bool bDraggedOntoAttachmentParent = true;
-		const auto& DragActors = Payload.GetData<TWeakObjectPtr<AActor>>(SceneOutliner::FWeakActorSelector());
+		const auto& DragActors = Payload.GetData<TWeakObjectPtr<AActor>>(ObjectMixerOutliner::FWeakActorSelector());
 		for (const auto& DragActorPtr : DragActors)
 		{
 			AActor* DragActor = DragActorPtr.Get();
@@ -2386,7 +2383,7 @@ FSceneOutlinerDragValidationInfo FObjectMixerOutlinerMode::ValidateDrop(const IS
 		{
 			const ULevelInstanceSubsystem* LevelInstanceSubsystem = RepresentingWorld->GetSubsystem<ULevelInstanceSubsystem>();
 			// Iterate over all the actors that have been dragged
-			for (const TWeakObjectPtr<AActor>& WeakActor : Payload.GetData<TWeakObjectPtr<AActor>>(SceneOutliner::FWeakActorSelector()))
+			for (const TWeakObjectPtr<AActor>& WeakActor : Payload.GetData<TWeakObjectPtr<AActor>>(ObjectMixerOutliner::FWeakActorSelector()))
 			{
 				const AActor* Actor = WeakActor.Get();
 
@@ -2481,7 +2478,7 @@ void FObjectMixerOutlinerMode::OnDrop(ISceneOutlinerTreeItem& DropTarget, const
 		{
 			const FScopedTransaction Transaction(LOCTEXT("UndoAction_DetachActors", "Detach actors"));
 
-			TArray<TWeakObjectPtr<AActor>> DraggedActors = Payload.GetData<TWeakObjectPtr<AActor>>(SceneOutliner::FWeakActorSelector());
+			TArray<TWeakObjectPtr<AActor>> DraggedActors = Payload.GetData<TWeakObjectPtr<AActor>>(ObjectMixerOutliner::FWeakActorSelector());
 			for (const auto& WeakActor : DraggedActors)
 			{
 				if (auto* DragActor = WeakActor.Get())
@@ -2517,7 +2514,7 @@ void FObjectMixerOutlinerMode::OnDrop(ISceneOutlinerTreeItem& DropTarget, const
 				Payload.ForEachItem<FFolderTreeItem>(MoveToDestination);
 
 				// Since target root is directly the Level Instance, clear folder path
-				TArray<AActor*> DraggedActors = Payload.GetData<AActor*>(SceneOutliner::FActorSelector());
+				TArray<AActor*> DraggedActors = Payload.GetData<AActor*>(ObjectMixerOutliner::FActorSelector());
 				for (auto& Actor : DraggedActors)
 				{
 					Actor->SetFolderPath_Recursively(FName());
@@ -2551,7 +2548,7 @@ void FObjectMixerOutlinerMode::OnDrop(ISceneOutlinerTreeItem& DropTarget, const
 					}
 				};
 
-				TArray<TWeakObjectPtr<AActor>> DraggedActors = Payload.GetData<TWeakObjectPtr<AActor>>(SceneOutliner::FWeakActorSelector());
+				TArray<TWeakObjectPtr<AActor>> DraggedActors = Payload.GetData<TWeakObjectPtr<AActor>>(ObjectMixerOutliner::FWeakActorSelector());
 				//@TODO: Should create a menu for each component that contains sockets, or have some form of disambiguation within the menu (like a fully qualified path)
 				// Instead, we currently only display the sockets on the root component
 				USceneComponent* Component = DropActor->GetRootComponent();
@@ -2735,7 +2732,7 @@ void FObjectMixerOutlinerMode::OnDrop(ISceneOutlinerTreeItem& DropTarget, const
 FFolder FObjectMixerOutlinerMode::CreateNewFolder()
 {
 	const FScopedTransaction Transaction(LOCTEXT("UndoAction_CreateFolder", "Create Folder"));
-	TArray<FFolder> SelectedFolders = SceneOutliner->GetSelection().GetData<FFolder>(SceneOutliner::FFolderPathSelector());
+	TArray<FFolder> SelectedFolders = SceneOutliner->GetSelection().GetData<FFolder>(ObjectMixerOutliner::FFolderPathSelector());
 	const FFolder NewFolderName = FActorFolders::Get().GetDefaultFolderForSelection(*RepresentingWorld, &SelectedFolders);
 	FActorFolders::Get().CreateFolderContainingSelection(*RepresentingWorld, NewFolderName);
 
@@ -2767,7 +2764,7 @@ bool FObjectMixerOutlinerMode::ReparentItemToFolder(const FFolder& FolderPath, c
 	return false;
 }
 
-namespace ActorBrowsingModeUtils
+namespace ObjectMixerActorBrowsingModeUtils
 {
 	static void RecursiveFolderExpandChildren(SSceneOutliner* SceneOutliner, const FSceneOutlinerTreeItemPtr& Item)
 	{
@@ -2879,7 +2876,7 @@ namespace ActorBrowsingModeUtils
 		}
 	};
 
-	bool CanChangePinnedStates(const TArray<FSceneOutlinerTreeItemPtr>& InItems)
+	static bool CanChangePinnedStates(const TArray<FSceneOutlinerTreeItemPtr>& InItems)
 	{
 		for (const FSceneOutlinerTreeItemPtr& Item : InItems)
 		{
@@ -2906,7 +2903,7 @@ void FObjectMixerOutlinerMode::SelectFoldersDescendants(const TArray<FFolderTree
 		SceneOutliner->SetItemExpansion(FolderPtr, true);
 		if (!bSelectImmediateChildrenOnly)
 		{
-			ActorBrowsingModeUtils::RecursiveFolderExpandChildren(SceneOutliner, FolderPtr);
+			ObjectMixerActorBrowsingModeUtils::RecursiveFolderExpandChildren(SceneOutliner, FolderPtr);
 		}
 	}
 
@@ -2915,7 +2912,7 @@ void FObjectMixerOutlinerMode::SelectFoldersDescendants(const TArray<FFolderTree
 
 	for (FFolderTreeItem* Folder : FolderItems)
 	{
-		ActorBrowsingModeUtils::RecursiveActorSelect(SceneOutliner, Folder->AsShared(), bSelectImmediateChildrenOnly);
+		ObjectMixerActorBrowsingModeUtils::RecursiveActorSelect(SceneOutliner, Folder->AsShared(), bSelectImmediateChildrenOnly);
 	}
 
 	GEditor->GetSelectedActors()->EndBatchSelectOperation(/*bNotify*/false);
@@ -2924,7 +2921,7 @@ void FObjectMixerOutlinerMode::SelectFoldersDescendants(const TArray<FFolderTree
 
 bool FObjectMixerOutlinerMode::CanPinItems(const TArray<FSceneOutlinerTreeItemPtr>& InItems) const
 {
-	return ActorBrowsingModeUtils::CanChangePinnedStates(InItems);
+	return ObjectMixerActorBrowsingModeUtils::CanChangePinnedStates(InItems);
 }
 
 void FObjectMixerOutlinerMode::PinItems(const TArray<FSceneOutlinerTreeItemPtr>& InItems)
@@ -2938,7 +2935,7 @@ void FObjectMixerOutlinerMode::PinItems(const TArray<FSceneOutlinerTreeItemPtr>&
 	TArray<FGuid> ActorsToPin;
 	// If Unloaded actors are hidden and we are pinning folders we need to find them through FActorFolders
 	const bool bSearchForHiddenUnloadedActors = bHideUnloadedActors;
-	ActorBrowsingModeUtils::RecursiveAddItemsToActorGuidList(InItems, ActorsToPin, bSearchForHiddenUnloadedActors);
+	ObjectMixerActorBrowsingModeUtils::RecursiveAddItemsToActorGuidList(InItems, ActorsToPin, bSearchForHiddenUnloadedActors);
 
 	if (ActorsToPin.Num())
 	{
@@ -2971,7 +2968,7 @@ void FObjectMixerOutlinerMode::PinItems(const TArray<FSceneOutlinerTreeItemPtr>&
 
 bool FObjectMixerOutlinerMode::CanUnpinItems(const TArray<FSceneOutlinerTreeItemPtr>& InItems) const
 {
-	return ActorBrowsingModeUtils::CanChangePinnedStates(InItems);
+	return ObjectMixerActorBrowsingModeUtils::CanChangePinnedStates(InItems);
 }
 
 void FObjectMixerOutlinerMode::UnpinItems(const TArray<FSceneOutlinerTreeItemPtr>& InItems)
@@ -2985,7 +2982,7 @@ void FObjectMixerOutlinerMode::UnpinItems(const TArray<FSceneOutlinerTreeItemPtr
 	TArray<FGuid> ActorsToUnpin;
 	// No need to search for hidden unloaded actors when unloading
 	const bool bSearchForHiddenUnloadedActors = false;
-	ActorBrowsingModeUtils::RecursiveAddItemsToActorGuidList(InItems, ActorsToUnpin, bSearchForHiddenUnloadedActors);
+	ObjectMixerActorBrowsingModeUtils::RecursiveAddItemsToActorGuidList(InItems, ActorsToUnpin, bSearchForHiddenUnloadedActors);
 
 	if (ActorsToUnpin.Num())
 	{
@@ -3100,7 +3097,7 @@ void FObjectMixerOutlinerMode::OnDuplicateActorsBegin()
 	// Only a callback in actor browsing mode
 	FFolder::FRootObject CommonRootObject;
 	TArray<FName> SelectedFolderPaths;
-	FFolder::GetFolderPathsAndCommonRootObject(SceneOutliner->GetSelection().GetData<FFolder>(SceneOutliner::FFolderPathSelector()), SelectedFolderPaths, CommonRootObject);
+	FFolder::GetFolderPathsAndCommonRootObject(SceneOutliner->GetSelection().GetData<FFolder>(ObjectMixerOutliner::FFolderPathSelector()), SelectedFolderPaths, CommonRootObject);
 	SceneOutliner->PasteFoldersBegin(SelectedFolderPaths);
 }
 
diff --git a/Engine/Plugins/Editor/ObjectMixer/ObjectMixer/Source/ObjectMixer/Public/Views/List/Modes/ObjectMixerOutlinerMode.h b/Engine/Plugins/Editor/ObjectMixer/ObjectMixer/Source/ObjectMixer/Public/Views/List/Modes/ObjectMixerOutlinerMode.h
--- a/Engine/Plugins/Editor/ObjectMixer/ObjectMixer/Source/ObjectMixer/Public/Views/List/Modes/ObjectMixerOutlinerMode.h
+++ b/Engine/Plugins/Editor/ObjectMixer/ObjectMixer/Source/ObjectMixer/Public/Views/List/Modes/ObjectMixerOutlinerMode.h
@@ -25,6 +25,29 @@ namespace ObjectMixerOutliner
 	{
 		bool operator()(const TWeakPtr<ISceneOutlinerTreeItem>& Item, TObjectPtr<UActorComponent>& DataOut) const;
 	};
+
+	/** Functor which can be used to get weak actor pointers from a selection */
+	struct FWeakActorSelector
+	{
+		bool operator()(const TWeakPtr<ISceneOutlinerTreeItem>& Item, TWeakObjectPtr<AActor>& DataOut) const;
+	};
+
+	/** Functor which can be used to get actors from a selection including component parents */
+	struct FActorSelector
+	{
+		bool operator()(const TWeakPtr<ISceneOutlinerTreeItem>& Item, AActor*& ActorPtrOut) const;
+	};
+
+	/** Functor which can be used to get actor descriptors from a selection  */
+	struct FActorDescSelector
+	{
+		bool operator()(const TWeakPtr<ISceneOutlinerTreeItem>& Item, FWorldPartitionActorDesc*& ActorDescPtrOut) const;
+	};
+
+	struct FFolderPathSelector
+	{
+		bool operator()(TWeakPtr<ISceneOutlinerTreeItem> Item, FFolder& DataOut) const;
+	};
 }
 
 USTRUCT()
