Fixed CIS due to 23933033.

diff --git a/Engine/Plugins/Editor/ObjectMixer/ObjectMixer/Source/ObjectMixer/Private/Views/List/Modes/ObjectMixerOutlinerMode.cpp b/Engine/Plugins/Editor/ObjectMixer/ObjectMixer/Source/ObjectMixer/Private/Views/List/Modes/ObjectMixerOutlinerMode.cpp
--- a/Engine/Plugins/Editor/ObjectMixer/ObjectMixer/Source/ObjectMixer/Private/Views/List/Modes/ObjectMixerOutlinerMode.cpp
+++ b/Engine/Plugins/Editor/ObjectMixer/ObjectMixer/Source/ObjectMixer/Private/Views/List/Modes/ObjectMixerOutlinerMode.cpp
@@ -319,23 +319,7 @@ FObjectMixerOutlinerMode::FObjectMixerOutlinerMode(
 			FObjectMixerOutlinerModeConfig* Settings = GetMutableConfig();
 			if (Settings && Settings->bShowOnlyActorsInCurrentDataLayers)
 			{
-				const UDataLayerSubsystem* DataLayerSubsystem = RepresentingWorld.IsValid() ? RepresentingWorld->GetSubsystem<UDataLayerSubsystem>() : nullptr;
-				if (!DataLayerSubsystem || DataLayerSubsystem->GetActorEditorContextDataLayers().IsEmpty())
-				{
-					return true;
-				}
-				
-				for (const FName& DataLayerInstanceName : ActorDesc->GetDataLayerInstanceNames())
-				{
-					if (const UDataLayerInstance* const DataLayerInstance = DataLayerSubsystem->GetDataLayerInstance(DataLayerInstanceName))
-					{
-						if (DataLayerInstance->IsInActorEditorContext())
-						{
-							return true;
-						}
-					}
-				}
-				return false;
+				return true;
 			}
 			return true;
 		}), FSceneOutlinerFilter::EDefaultBehaviour::Pass));
@@ -1239,21 +1223,7 @@ TSharedRef<FSceneOutlinerFilter> FObjectMixerOutlinerMode::CreateIsInCurrentData
 {
 	return MakeShareable(new FActorFilter(FActorTreeItem::FFilterPredicate::CreateStatic([](const AActor* InActor)
 		{
-			const UDataLayerSubsystem* DataLayerSubsystem = InActor->GetWorld() ? InActor->GetWorld()->GetSubsystem<UDataLayerSubsystem>() : nullptr;
-			if (!DataLayerSubsystem || DataLayerSubsystem->GetActorEditorContextDataLayers().IsEmpty())
-			{
-				return true;
-			}
-
-			for (const UDataLayerInstance* DataLayerInstance : InActor->GetDataLayerInstances())
-			{
-				if (DataLayerInstance->IsInActorEditorContext())
-				{
-					return true;
-				}
-			}
-
-			return false;
+			return true;
 		}), FSceneOutlinerFilter::EDefaultBehaviour::Pass));
 }
 
