Fixed remaining CIS due to 23933033

diff --git a/Engine/Plugins/Editor/ObjectMixer/ObjectMixer/Source/ObjectMixer/Private/Views/List/SObjectMixerEditorList.cpp b/Engine/Plugins/Editor/ObjectMixer/ObjectMixer/Source/ObjectMixer/Private/Views/List/SObjectMixerEditorList.cpp
--- a/Engine/Plugins/Editor/ObjectMixer/ObjectMixer/Source/ObjectMixer/Private/Views/List/SObjectMixerEditorList.cpp
+++ b/Engine/Plugins/Editor/ObjectMixer/ObjectMixer/Source/ObjectMixer/Private/Views/List/SObjectMixerEditorList.cpp
@@ -49,6 +49,9 @@ const FName SObjectMixerEditorList::ItemNameColumnName(TEXT("Builtin_Name"));
 const FName SObjectMixerEditorList::EditorVisibilityColumnName(TEXT("Builtin_EditorVisibility"));
 const FName SObjectMixerEditorList::EditorVisibilitySoloColumnName(TEXT("Builtin_EditorVisibilitySolo"));
 
+// Temporary until DataLayerSubsystem is replaced by DataLayerManager
+PRAGMA_DISABLE_DEPRECATION_WARNINGS
+
 void SObjectMixerEditorList::Construct(const FArguments& InArgs, TSharedRef<FObjectMixerEditorList> ListModel)
 {
 	ListModelPtr = ListModel;
@@ -907,7 +910,7 @@ void SObjectMixerEditorList::CreateActorTextInfoColumns(UWorld *WorldPtr, FScene
 				{
 					if (const UActorDescContainer* ActorDescContainer = ActorDescItem->ActorDescHandle.Container.Get())
 					{
-						const UWorld* World = ActorDescContainer->GetWorld();
+						const UWorld* OwningWorld = ActorDescContainer->GetWorldPartition()->GetWorld();
 						if (const UDataLayerSubsystem* DataLayerSubsystem = UWorld::GetSubsystem<UDataLayerSubsystem>(World))
 						{
 							TSet<const UDataLayerInstance*> DataLayerInstances;
@@ -1415,4 +1418,6 @@ void SObjectMixerEditorList::PropagatePropertyChangesToSelectedRows()
 	PendingPropertyPropagations.Empty();
 }
 
+PRAGMA_ENABLE_DEPRECATION_WARNINGS
+
 #undef LOCTEXT_NAMESPACE
