Fixed remaining CIS due to 23933033

diff --git a/Engine/Plugins/Editor/ObjectMixer/ObjectMixer/Source/ObjectMixer/Private/Views/List/SObjectMixerEditorList.cpp b/Engine/Plugins/Editor/ObjectMixer/ObjectMixer/Source/ObjectMixer/Private/Views/List/SObjectMixerEditorList.cpp
--- a/Engine/Plugins/Editor/ObjectMixer/ObjectMixer/Source/ObjectMixer/Private/Views/List/SObjectMixerEditorList.cpp
+++ b/Engine/Plugins/Editor/ObjectMixer/ObjectMixer/Source/ObjectMixer/Private/Views/List/SObjectMixerEditorList.cpp
@@ -910,7 +910,7 @@ void SObjectMixerEditorList::CreateActorTextInfoColumns(UWorld *WorldPtr, FScene
 				{
 					if (const UActorDescContainer* ActorDescContainer = ActorDescItem->ActorDescHandle.Container.Get())
 					{
-						const UWorld* OwningWorld = ActorDescContainer->GetWorldPartition()->GetWorld();
+						const UWorld* World = ActorDescContainer->GetWorldPartition()->GetWorld();
 						if (const UDataLayerSubsystem* DataLayerSubsystem = UWorld::GetSubsystem<UDataLayerSubsystem>(World))
 						{
 							TSet<const UDataLayerInstance*> DataLayerInstances;
diff --git a/Engine/Source/Runtime/Engine/Public/WorldPartition/DataLayer/DataLayerSubsystem.h b/Engine/Source/Runtime/Engine/Public/WorldPartition/DataLayer/DataLayerSubsystem.h
--- a/Engine/Source/Runtime/Engine/Public/WorldPartition/DataLayer/DataLayerSubsystem.h
+++ b/Engine/Source/Runtime/Engine/Public/WorldPartition/DataLayer/DataLayerSubsystem.h
@@ -211,5 +211,5 @@ template<class T>
 TArray<const UDataLayerInstance*> UDataLayerSubsystem::GetDataLayerInstances(const TArray<T>& InDataLayerIdentifiers, const ULevel* InLevelContext /* = nullptr */) const
 {
 	const UDataLayerManager* DataLayerManager = UDataLayerManager::GetDataLayerManager(GetWorld());
-	return DataLayerManager ? DataLayerManager->GetDataLayerInstances(InDataLayerIdentifiers, InLevelContext) : TArray<const UDataLayerInstance*>();
+	return DataLayerManager ? DataLayerManager->GetDataLayerInstances(InDataLayerIdentifiers) : TArray<const UDataLayerInstance*>();
 }
\ No newline at end of file
