Converted ObjectMixer to use DataLayerManager

diff --git a/Engine/Plugins/Editor/ObjectMixer/ObjectMixer/Source/ObjectMixer/Private/Views/List/Modes/ObjectMixerOutlinerMode.cpp b/Engine/Plugins/Editor/ObjectMixer/ObjectMixer/Source/ObjectMixer/Private/Views/List/Modes/ObjectMixerOutlinerMode.cpp
--- a/Engine/Plugins/Editor/ObjectMixer/ObjectMixer/Source/ObjectMixer/Private/Views/List/Modes/ObjectMixerOutlinerMode.cpp
+++ b/Engine/Plugins/Editor/ObjectMixer/ObjectMixer/Source/ObjectMixer/Private/Views/List/Modes/ObjectMixerOutlinerMode.cpp
@@ -57,7 +57,6 @@
 #include "WorldPartition/IWorldPartitionEditorModule.h"
 #include "WorldPartition/WorldPartition.h"
 #include "WorldPartition/WorldPartitionSubsystem.h"
-#include "WorldPartition/DataLayer/DataLayerSubsystem.h"
 
 static int32 GSceneOutlinerAutoRepresentingWorldNetModeForObjectMixer = NM_Client;
 static FAutoConsoleVariableRef CVarAutoRepresentingWorldNetMode(
diff --git a/Engine/Plugins/Editor/ObjectMixer/ObjectMixer/Source/ObjectMixer/Private/Views/List/SObjectMixerEditorList.cpp b/Engine/Plugins/Editor/ObjectMixer/ObjectMixer/Source/ObjectMixer/Private/Views/List/SObjectMixerEditorList.cpp
--- a/Engine/Plugins/Editor/ObjectMixer/ObjectMixer/Source/ObjectMixer/Private/Views/List/SObjectMixerEditorList.cpp
+++ b/Engine/Plugins/Editor/ObjectMixer/ObjectMixer/Source/ObjectMixer/Private/Views/List/SObjectMixerEditorList.cpp
@@ -41,7 +41,7 @@
 #include "WorldPartition/WorldPartitionActorDesc.h"
 #include "WorldPartition/ContentBundle/ContentBundleDescriptor.h"
 #include "WorldPartition/ContentBundle/ContentBundleEngineSubsystem.h"
-#include "WorldPartition/DataLayer/DataLayerSubsystem.h"
+#include "WorldPartition/DataLayer/DataLayerManager.h"
 
 #define LOCTEXT_NAMESPACE "SObjectMixerEditorList"
 
@@ -49,9 +49,6 @@ const FName SObjectMixerEditorList::ItemNameColumnName(TEXT("Builtin_Name"));
 const FName SObjectMixerEditorList::EditorVisibilityColumnName(TEXT("Builtin_EditorVisibility"));
 const FName SObjectMixerEditorList::EditorVisibilitySoloColumnName(TEXT("Builtin_EditorVisibilitySolo"));
 
-// Temporary until DataLayerSubsystem is replaced by DataLayerManager
-PRAGMA_DISABLE_DEPRECATION_WARNINGS
-
 void SObjectMixerEditorList::Construct(const FArguments& InArgs, TSharedRef<FObjectMixerEditorList> ListModel)
 {
 	ListModelPtr = ListModel;
@@ -910,12 +907,12 @@ void SObjectMixerEditorList::CreateActorTextInfoColumns(UWorld *WorldPtr, FScene
 				{
 					if (const UActorDescContainer* ActorDescContainer = ActorDescItem->ActorDescHandle.Container.Get())
 					{
-						const UWorld* World = ActorDescContainer->GetWorldPartition()->GetWorld();
-						if (const UDataLayerSubsystem* DataLayerSubsystem = UWorld::GetSubsystem<UDataLayerSubsystem>(World))
+						const UWorld* OwningWorld = ActorDescContainer->GetWorldPartition()->GetWorld();
+						if (const UDataLayerManager* DataLayerManager = UDataLayerManager::GetDataLayerManager(OwningWorld))
 						{
 							TSet<const UDataLayerInstance*> DataLayerInstances;
-							DataLayerInstances.Append(DataLayerSubsystem->GetDataLayerInstances(ActorDesc->GetDataLayerInstanceNames()));
-							if (ULevelInstanceSubsystem* LevelInstanceSubsystem = UWorld::GetSubsystem<ULevelInstanceSubsystem>(World))
+							DataLayerInstances.Append(DataLayerManager->GetDataLayerInstances(ActorDesc->GetDataLayerInstanceNames()));
+							if (ULevelInstanceSubsystem* LevelInstanceSubsystem = UWorld::GetSubsystem<ULevelInstanceSubsystem>(OwningWorld))
 							{
 								UWorld* OuterWorld = ActorDescContainer->GetTypedOuter<UWorld>();
 								// Add parent container Data Layer Instances
@@ -1418,6 +1415,4 @@ void SObjectMixerEditorList::PropagatePropertyChangesToSelectedRows()
 	PendingPropertyPropagations.Empty();
 }
 
-PRAGMA_ENABLE_DEPRECATION_WARNINGS
-
 #undef LOCTEXT_NAMESPACE
