Object Mixer:

-Fix: Rebuild List Option not working

diff --git a/Engine/Plugins/Editor/ObjectMixer/ObjectMixer/Source/ObjectMixer/Private/Views/List/Modes/ObjectMixerOutlinerMode.cpp b/Engine/Plugins/Editor/ObjectMixer/ObjectMixer/Source/ObjectMixer/Private/Views/List/Modes/ObjectMixerOutlinerMode.cpp
--- a/Engine/Plugins/Editor/ObjectMixer/ObjectMixer/Source/ObjectMixer/Private/Views/List/Modes/ObjectMixerOutlinerMode.cpp
+++ b/Engine/Plugins/Editor/ObjectMixer/ObjectMixer/Source/ObjectMixer/Private/Views/List/Modes/ObjectMixerOutlinerMode.cpp
@@ -1168,12 +1168,6 @@ void FObjectMixerOutlinerMode::CreateViewContent(FMenuBuilder& MenuBuilder)
 				})
 			)
 		);
-
-		MenuBuilder.AddMenuEntry(
-			LOCTEXT("RebuildListMenuOption","Rebuild List"), 
-			LOCTEXT("RebuildListMenuOptionTooltip","Force the active list to be rebuilt. Useful if the list doesn't automatically refresh."),
-			FSlateIcon(),
-			FUIAction(FExecuteAction::CreateSP(GetListModelPtr().Pin().ToSharedRef(), &FObjectMixerEditorList::RequestRebuildList)));
 	}
 	MenuBuilder.EndSection();
 
