Object Mixer:

-Reimplement Hybrid Rows

diff --git a/Engine/Plugins/Editor/ObjectMixer/ObjectMixer/Source/ObjectMixer/Private/Views/List/Modes/ObjectMixerOutlinerHierarchy.cpp b/Engine/Plugins/Editor/ObjectMixer/ObjectMixer/Source/ObjectMixer/Private/Views/List/Modes/ObjectMixerOutlinerHierarchy.cpp
--- a/Engine/Plugins/Editor/ObjectMixer/ObjectMixer/Source/ObjectMixer/Private/Views/List/Modes/ObjectMixerOutlinerHierarchy.cpp
+++ b/Engine/Plugins/Editor/ObjectMixer/ObjectMixer/Source/ObjectMixer/Private/Views/List/Modes/ObjectMixerOutlinerHierarchy.cpp
@@ -4,6 +4,7 @@
 
 #include "Views/List/ObjectMixerEditorList.h"
 #include "Views/List/Modes/ObjectMixerOutlinerMode.h"
+#include "Views/List/ObjectMixerUtils.h"
 #include "Views/List/RowTypes/ObjectMixerEditorListRowActor.h"
 #include "Views/List/RowTypes/ObjectMixerEditorListRowComponent.h"
 #include "Views/List/RowTypes/ObjectMixerEditorListRowFolder.h"
@@ -332,7 +333,7 @@ FSceneOutlinerTreeItemPtr FObjectMixerOutlinerHierarchy::FindOrCreateParentItem(
 				}
 			}
 
-			// Parent Actor (Actor attachement / parenting)
+			// Parent Actor (Actor attachment / parenting)
 			const FGuid& ParentActorGuid = ActorDesc->GetSceneOutlinerParent();
 			if (ParentActorGuid.IsValid())
 			{
@@ -463,6 +464,85 @@ void FObjectMixerOutlinerHierarchy::CreateFolderChild(const FFolder& Folder, UWo
     }
 }
 
+void FObjectMixerOutlinerHierarchy::EvaluateActorsByComponentsAndCreateActorItems(
+	AActor* Actor, TArray<FSceneOutlinerTreeItemPtr>& OutItems) const
+{
+	if (!Actor)
+	{
+		return;
+	}
+
+	if (FSceneOutlinerTreeItemPtr ActorItem =
+			Mode->CreateItemFor<FObjectMixerEditorListRowActor>(
+				FObjectMixerEditorListRowActor(Actor, GetCastedMode()->GetSceneOutliner())))
+	{	
+		// Create all component items
+		TArray<FSceneOutlinerTreeItemPtr> ComponentItems;
+		CreateComponentItems(Actor, ComponentItems);
+
+		if (ComponentItems.Num()) // Only add items if we have components
+		{	
+			OutItems.Add(ActorItem);
+			
+			if (ComponentItems.Num() == 1) // Create hybrid row
+			{
+				if (FObjectMixerEditorListRowActor* AsActorRow = FObjectMixerUtils::AsActorRow(ActorItem))
+				{
+					const FComponentTreeItem* ComponentItem = ComponentItems[0]->CastTo<FComponentTreeItem>();
+
+					if (ComponentItem && ComponentItem->Component.IsValid())
+					{
+						AsActorRow->RowData.SetHybridComponent(ComponentItem->Component.Get());
+					}
+				}
+			}
+			else // Create normal Actor->Component hierarchy if there are multiple matching components
+			{
+				OutItems.Append(ComponentItems);
+			}
+		}
+	}
+}
+
+void FObjectMixerOutlinerHierarchy::ForEachActorInLevel(
+	AActor* Actor, const ULevelInstanceSubsystem* LevelInstanceSubsystem, TArray<FSceneOutlinerTreeItemPtr>& OutItems) const
+{
+	if (!Actor)
+	{
+		return;
+	}
+	
+	// If we are not showing LevelInstances, LevelInstance sub actor items should not be created unless they belong to a LevelInstance which is being edited
+	if (LevelInstanceSubsystem)
+	{
+		if (const ILevelInstanceInterface* ParentLevelInstance = LevelInstanceSubsystem->GetParentLevelInstance(Actor))
+		{
+			if (!bShowingLevelInstances && !ParentLevelInstance->IsEditing())
+			{
+				return;
+			}
+		}
+	}
+
+	// If actor class does not match filter and we only want to show actors that have valid subobjects
+	if (bShowingOnlyActorWithValidComponents && !DoesWorldObjectHaveAcceptableClass(Actor))
+	{
+		EvaluateActorsByComponentsAndCreateActorItems(Actor, OutItems);
+	}
+	else // Create actor rows anyway
+	{
+		if (FSceneOutlinerTreeItemPtr ActorItem =
+			Mode->CreateItemFor<FObjectMixerEditorListRowActor>(
+				FObjectMixerEditorListRowActor(Actor, GetCastedMode()->GetSceneOutliner())))
+		{
+			OutItems.Add(ActorItem);
+
+			// Create all component items
+			CreateComponentItems(Actor, OutItems);
+		}
+	}	
+}
+
 void FObjectMixerOutlinerHierarchy::CreateWorldChildren(UWorld* World, TArray<FSceneOutlinerTreeItemPtr>& OutItems) const
 {
 	check(World);
@@ -481,73 +561,36 @@ void FObjectMixerOutlinerHierarchy::CreateWorldChildren(UWorld* World, TArray<FS
 	
 	for (const ULevel* Level : World->GetLevels())
 	{
-		if (Level && Level->bIsVisible)
+		if (!Level || !Level->bIsVisible)
 		{
-			ForEachObjectWithOuter(Level,
-			  [&](UObject* Object)
-		  {
-				if (AActor* Actor = Cast<AActor>(Object))
-				{
-					// If we are not showing LevelInstances, LevelInstance sub actor items should not be created unless they belong to a LevelInstance which is being edited
-					if (LevelInstanceSubsystem)
-					{
-						if (const ILevelInstanceInterface* ParentLevelInstance = LevelInstanceSubsystem->GetParentLevelInstance(Actor))
-						{
-							if (!bShowingLevelInstances && !ParentLevelInstance->IsEditing())
-							{
-								return;
-							}
-						}
-					}
-					
-					if (FSceneOutlinerTreeItemPtr ActorItem =
-						Mode->CreateItemFor<FObjectMixerEditorListRowActor>(
-							FObjectMixerEditorListRowActor(Actor, GetCastedMode()->GetSceneOutliner())))
-					{
-						if (bShowingOnlyActorWithValidComponents)
-						{
-							const int32 InsertLocation = OutItems.Num();
-
-							// Create all component items
-							CreateComponentItems(Actor, OutItems);
-
-							if (OutItems.Num() != InsertLocation || DoesWorldObjectHaveAcceptableClass(Actor))
-							{
-								// Add the actor before the components
-								OutItems.Insert(ActorItem, InsertLocation);
-							}
-						}
-						else
-						{
-							OutItems.Add(ActorItem);
+			continue;
+		}
 
-							// Create all component items
-							CreateComponentItems(Actor, OutItems);
-						}
-					}
-				}
-				// Since components are created per actor already, we only are concerned with other matching UObjects
-				else if (!Object->IsA(UActorComponent::StaticClass())) 
+		ForEachObjectWithOuter(Level,
+		  [&](UObject* Object)
+		{
+			if (AActor* Actor = Cast<AActor>(Object))
+			{
+				ForEachActorInLevel(Actor, LevelInstanceSubsystem, OutItems);
+			}
+			// Since components are created per actor already, we only are concerned with other matching UObjects
+			else if (!Object->IsA(UActorComponent::StaticClass()) && DoesWorldObjectHaveAcceptableClass(Object)) 
+			{
+				if (FSceneOutlinerTreeItemPtr ObjectItem =
+					Mode->CreateItemFor<FObjectMixerEditorListRowUObject>(
+						FObjectMixerEditorListRowUObject(Object, GetCastedMode()->GetSceneOutliner())))
 				{
-					if (DoesWorldObjectHaveAcceptableClass(Object))
-					{
-						if (FSceneOutlinerTreeItemPtr ObjectItem =
-							Mode->CreateItemFor<FObjectMixerEditorListRowUObject>(
-								FObjectMixerEditorListRowUObject(Object, GetCastedMode()->GetSceneOutliner())))
-						{
-							OutItems.Add(ObjectItem);
-						}
-					}
+					OutItems.Add(ObjectItem);
 				}
-		  });
-		}
+			}
+		});
 	}
 
 	if (bShowingUnloadedActors)
 	{
 		if (UWorldPartitionSubsystem* WorldPartitionSubsystem = UWorld::GetSubsystem<UWorldPartitionSubsystem>(World))
 		{
-			WorldPartitionSubsystem->ForEachWorldPartition([this, LevelInstanceSubsystem, &OutItems](UWorldPartition* WorldPartition)
+			auto ForEachWorldPartition = [this, LevelInstanceSubsystem, &OutItems](UWorldPartition* WorldPartition)
 			{
 				// Skip unloaded actors if they are part of a non editing level instance and the outliner hides the content of level instances
 				if (!bShowingLevelInstances)
@@ -560,11 +603,13 @@ void FObjectMixerOutlinerHierarchy::CreateWorldChildren(UWorld* World, TArray<FS
 						return true;
 					}
 				}
-				FWorldPartitionHelpers::ForEachActorDesc(WorldPartition, [this, &OutItems](const FWorldPartitionActorDesc* ActorDesc)
+				FWorldPartitionHelpers::ForEachActorDesc(WorldPartition,
+					[this, &OutItems](const FWorldPartitionActorDesc* ActorDesc)
 				{
 					if (ActorDesc != nullptr && !ActorDesc->IsLoaded(true))
 					{
-						if (const FSceneOutlinerTreeItemPtr ActorDescItem = Mode->CreateItemFor<FActorDescTreeItem>(FActorDescTreeItem(ActorDesc->GetGuid(), ActorDesc->GetContainer())))
+						if (const FSceneOutlinerTreeItemPtr ActorDescItem =
+							Mode->CreateItemFor<FActorDescTreeItem>(FActorDescTreeItem(ActorDesc->GetGuid(), ActorDesc->GetContainer())))
 						{
 							OutItems.Add(ActorDescItem);
 						}
@@ -572,7 +617,9 @@ void FObjectMixerOutlinerHierarchy::CreateWorldChildren(UWorld* World, TArray<FS
 					return true;
 				});
 				return true;
-			});
+			};
+			
+			WorldPartitionSubsystem->ForEachWorldPartition(ForEachWorldPartition);
 		}
 	}
 }
diff --git a/Engine/Plugins/Editor/ObjectMixer/ObjectMixer/Source/ObjectMixer/Private/Views/List/ObjectMixerEditorListRowData.cpp b/Engine/Plugins/Editor/ObjectMixer/ObjectMixer/Source/ObjectMixer/Private/Views/List/ObjectMixerEditorListRowData.cpp
--- a/Engine/Plugins/Editor/ObjectMixer/ObjectMixer/Source/ObjectMixer/Private/Views/List/ObjectMixerEditorListRowData.cpp
+++ b/Engine/Plugins/Editor/ObjectMixer/ObjectMixer/Source/ObjectMixer/Private/Views/List/ObjectMixerEditorListRowData.cpp
@@ -6,8 +6,8 @@
 #include "Views/List/ObjectMixerUtils.h"
 #include "Views/List/SObjectMixerEditorList.h"
 
-#include "ScopedTransaction.h"
 #include "GameFramework/Actor.h"
+#include "ScopedTransaction.h"
 
 const TArray<TObjectPtr<UObjectMixerObjectFilter>>& FObjectMixerEditorListRowData::GetObjectFilterInstances() const
 {
@@ -43,16 +43,6 @@ void FObjectMixerEditorListRowData::SetIsTreeViewItemExpanded(const TSharedRef<I
 	GetListView()->SetTreeViewItemExpanded(InRow, bNewExpanded);
 }
 
-bool FObjectMixerEditorListRowData::GetDoesRowPassFilters() const
-{
-	return bDoesRowPassFilters;
-}
-
-void FObjectMixerEditorListRowData::SetDoesRowPassFilters(const bool bPass)
-{
-	bDoesRowPassFilters = bPass;
-}
-
 bool FObjectMixerEditorListRowData::GetIsSelected(const TSharedRef<ISceneOutlinerTreeItem> InRow)
 {
 	if (SObjectMixerEditorList* ListView = GetListView())
@@ -92,7 +82,7 @@ bool FObjectMixerEditorListRowData::HasAtLeastOneChildThatIsNotSolo(const TShare
 	return false;
 }
 
-FText FObjectMixerEditorListRowData::GetDisplayName(TSharedPtr<ISceneOutlinerTreeItem> InTreeItem, const bool bIsHybridRow) const
+FText FObjectMixerEditorListRowData::GetDisplayName(TSharedPtr<ISceneOutlinerTreeItem> InTreeItem) const
 {
 	const FText Override = GetDisplayNameOverride();
 	if (!Override.IsEmpty())
@@ -104,7 +94,7 @@ FText FObjectMixerEditorListRowData::GetDisplayName(TSharedPtr<ISceneOutlinerTre
 	{
 		if (const TObjectPtr<UObject> Object = FObjectMixerUtils::GetRowObject(InTreeItem))
 		{
-			return Filter->GetRowDisplayName(Object, bIsHybridRow);
+			return Filter->GetRowDisplayName(Object, GetIsHybridRow());
 		}
 	}
 
@@ -116,16 +106,6 @@ SObjectMixerEditorList* FObjectMixerEditorListRowData::GetListView() const
 	return StaticCast<SObjectMixerEditorList*>(SceneOutlinerPtr);
 }
 
-EObjectMixerTreeViewMode FObjectMixerEditorListRowData::GetTreeViewMode()
-{
-	if (SObjectMixerEditorList* ListView = GetListView())
-	{
-		return ListView->GetTreeViewMode();
-	}
-
-	return EObjectMixerTreeViewMode::Folders;
-}
-
 TArray<TSharedPtr<ISceneOutlinerTreeItem>> FObjectMixerEditorListRowData::GetSelectedTreeViewItems() const
 {
 	if (SObjectMixerEditorList* ListView = GetListView())
@@ -177,21 +157,6 @@ void FObjectMixerEditorListRowData::ClearSoloRows() const
 	}
 }
 
-bool FObjectMixerEditorListRowData::GetIsItemOrHybridChildSelected(const TSharedRef<ISceneOutlinerTreeItem> InRow)
-{
-	if (GetIsSelected(InRow))
-	{
-		return true;
-	}
-	
-	if (const TSharedPtr<ISceneOutlinerTreeItem> HybridChild = FObjectMixerUtils::GetHybridChild(InRow).Pin())
-	{
-		return FObjectMixerUtils::GetRowData(HybridChild)->GetIsSelected(HybridChild.ToSharedRef());
-	}
-	
-	return false;
-}
-
 void SetValueOnSelectedItems(
 	const FString& ValueAsString, const TArray<TSharedPtr<ISceneOutlinerTreeItem>>& OtherSelectedItems,
 	const FName& PropertyName, const TSharedPtr<ISceneOutlinerTreeItem> PinnedItem,
@@ -204,23 +169,20 @@ void SetValueOnSelectedItems(
 		
 		for (const TSharedPtr<ISceneOutlinerTreeItem>& SelectedRow : OtherSelectedItems)
 		{
-			const TWeakPtr<ISceneOutlinerTreeItem> SelectedHybridRow = FObjectMixerUtils::GetHybridChild(SelectedRow);
-			const TSharedPtr<ISceneOutlinerTreeItem> RowToUse = SelectedHybridRow.IsValid() ? SelectedHybridRow.Pin() : SelectedRow;
-
-			if (RowToUse == PinnedItem)
+			if (SelectedRow == PinnedItem)
 			{
 				continue;
 			}
 
-			FObjectMixerEditorListRowData* RowData = FObjectMixerUtils::GetRowData(RowToUse);
+			FObjectMixerEditorListRowData* RowData = FObjectMixerUtils::GetRowData(SelectedRow);
 
 			// Skip folders
-			if (FObjectMixerUtils::AsFolderRow(RowToUse))
+			if (FObjectMixerUtils::AsFolderRow(SelectedRow))
 			{
 				continue;
 			}
 
-			UObject* ObjectToModify = FObjectMixerUtils::GetRowObject(RowToUse);
+			UObject* ObjectToModify = FObjectMixerUtils::GetRowObject(SelectedRow);
 			
 			if (IsValid(ObjectToModify))
 			{
@@ -228,7 +190,7 @@ void SetValueOnSelectedItems(
 			}
 			else
 			{
-				UE_LOG(LogObjectMixerEditor, Warning, TEXT("%hs: Row '%s' has no valid associated object to modify."), __FUNCTION__, *RowData->GetDisplayName(RowToUse).ToString());
+				UE_LOG(LogObjectMixerEditor, Warning, TEXT("%hs: Row '%s' has no valid associated object to modify."), __FUNCTION__, *RowData->GetDisplayName(SelectedRow).ToString());
 				return;
 			}
 		
@@ -263,28 +225,25 @@ void FObjectMixerEditorListRowData::PropagateChangesToSimilarSelectedRowProperti
 		return;
 	}
 
-	if (const TSharedPtr<ISceneOutlinerTreeItem> RowToUse = FObjectMixerUtils::GetHybridChildOrRowItemIfNull(InRow); RowToUse)
+	if (!GetIsSelected(InRow))
 	{
-		if (!GetIsItemOrHybridChildSelected(InRow))
-		{
-			return;
-		}
+		return;
+	}
 
-		FObjectMixerEditorListRowData* RowData = FObjectMixerUtils::GetRowData(RowToUse);
-		
-		const TWeakPtr<IPropertyHandle>* HandlePtr = RowData->PropertyNamesToHandles.Find(PropertyPropagationInfo.PropertyName);
-		if (HandlePtr->IsValid())
+	FObjectMixerEditorListRowData* RowData = FObjectMixerUtils::GetRowData(InRow);
+	
+	const TWeakPtr<IPropertyHandle>* HandlePtr = RowData->PropertyNamesToHandles.Find(PropertyPropagationInfo.PropertyName);
+	if (HandlePtr->IsValid())
+	{
+		const TArray<TSharedPtr<ISceneOutlinerTreeItem>> OtherSelectedItems = RowData->GetSelectedTreeViewItems();
+		if (OtherSelectedItems.Num())
 		{
-			const TArray<TSharedPtr<ISceneOutlinerTreeItem>> OtherSelectedItems = RowData->GetSelectedTreeViewItems();
-			if (OtherSelectedItems.Num())
-			{
-				FString ValueAsString;
-				(*HandlePtr).Pin()->GetValueAsFormattedString(ValueAsString);
-			
-				SetValueOnSelectedItems(
-					ValueAsString, OtherSelectedItems, PropertyPropagationInfo.PropertyName,
-					RowToUse, PropertyPropagationInfo.PropertyValueSetFlags);
-			}
+			FString ValueAsString;
+			(*HandlePtr).Pin()->GetValueAsFormattedString(ValueAsString);
+		
+			SetValueOnSelectedItems(
+				ValueAsString, OtherSelectedItems, PropertyPropagationInfo.PropertyName,
+				InRow, PropertyPropagationInfo.PropertyValueSetFlags);
 		}
 	}
 }
diff --git a/Engine/Plugins/Editor/ObjectMixer/ObjectMixer/Source/ObjectMixer/Private/Views/List/ObjectMixerUtils.cpp b/Engine/Plugins/Editor/ObjectMixer/ObjectMixer/Source/ObjectMixer/Private/Views/List/ObjectMixerUtils.cpp
--- a/Engine/Plugins/Editor/ObjectMixer/ObjectMixer/Source/ObjectMixer/Private/Views/List/ObjectMixerUtils.cpp
+++ b/Engine/Plugins/Editor/ObjectMixer/ObjectMixer/Source/ObjectMixer/Private/Views/List/ObjectMixerUtils.cpp
@@ -3,13 +3,11 @@
 #include "Views/List/ObjectMixerUtils.h"
 
 #include "ObjectMixerEditorSerializedData.h"
-#include "Views/List/SObjectMixerEditorList.h"
-
-#include "ScopedTransaction.h"
 #include "Views/List/RowTypes/ObjectMixerEditorListRowActor.h"
 #include "Views/List/RowTypes/ObjectMixerEditorListRowComponent.h"
 #include "Views/List/RowTypes/ObjectMixerEditorListRowFolder.h"
 #include "Views/List/RowTypes/ObjectMixerEditorListRowUObject.h"
+#include "Views/List/SObjectMixerEditorList.h"
 
 FObjectMixerEditorListRowFolder* FObjectMixerUtils::AsFolderRow(TSharedPtr<ISceneOutlinerTreeItem> InTreeItem)
 {
@@ -73,13 +71,6 @@ FObjectMixerEditorListRowComponent* FObjectMixerUtils::AsComponentRow(TSharedPtr
 
 FObjectMixerEditorListRowUObject* FObjectMixerUtils::AsObjectRow(TSharedPtr<ISceneOutlinerTreeItem> InTreeItem)
 {
-	// Sometimes the outliner reports type as the parent type, but we can be assured we're always using our type.
-	// For this reason, we check for our type first and if it's not our type, check parent type then cast.
-	if (FObjectMixerEditorListRowUObject* AsObject = InTreeItem->CastTo<FObjectMixerEditorListRowUObject>())
-	{
-		return AsObject;
-	}
-	
 	if (FObjectMixerEditorListRowUObject* AsObject = InTreeItem->CastTo<FObjectMixerEditorListRowUObject>())
 	{
 		if (IsValid(AsObject->ObjectPtr))
@@ -116,59 +107,15 @@ FObjectMixerEditorListRowData* FObjectMixerUtils::GetRowData(TSharedPtr<ISceneOu
 	return nullptr;
 }
 
-TWeakPtr<ISceneOutlinerTreeItem> FObjectMixerUtils::GetHybridChild(TSharedPtr<ISceneOutlinerTreeItem> InTreeItem)
-{
+UObject* FObjectMixerUtils::GetRowObject(TSharedPtr<ISceneOutlinerTreeItem> InTreeItem, const bool bGetHybridRowComponent)
+{	
 	if (const FObjectMixerEditorListRowActor* ActorRow = AsActorRow(InTreeItem))
 	{
-		if (const TWeakObjectPtr<AActor> ThisObject = ActorRow->Actor; ThisObject.IsValid())
+		if (bGetHybridRowComponent && ActorRow->RowData.GetIsHybridRow())
 		{
-			TArray<TWeakPtr<ISceneOutlinerTreeItem>> HybridCandidates;
-	
-			for (const TWeakPtr<ISceneOutlinerTreeItem>& ChildRow : InTreeItem->GetChildren())
-			{		
-				if (const FObjectMixerEditorListRowComponent* ComponentRow = AsComponentRow(ChildRow.Pin()))
-				{
-					if (TWeakObjectPtr<UActorComponent> ChildObject = ComponentRow->Component; ChildObject.IsValid())
-					{
-						if (ChildObject->GetOuter() == ThisObject)
-						{
-							HybridCandidates.Add(ChildRow);
-
-							if (HybridCandidates.Num() > 1)
-							{
-								// There can only be one row to hybrid with.
-								// If there's more than one candidate, don't hybrid.
-								break;
-							}
-						}
-					}
-				}
-			}
-		
-			if (HybridCandidates.Num() == 1)
-			{
-				return HybridCandidates[0];
-			}
+			return ActorRow->RowData.GetHybridComponent();
 		}
-	}
-	
-	return nullptr;
-}
-
-TSharedRef<ISceneOutlinerTreeItem> FObjectMixerUtils::GetHybridChildOrRowItemIfNull(const TSharedRef<ISceneOutlinerTreeItem> InRow)
-{
-	if (const TWeakPtr<ISceneOutlinerTreeItem> HybridChild = GetHybridChild(InRow); HybridChild.IsValid())
-	{
-		return HybridChild.Pin().ToSharedRef();
-	}
 		
-	return InRow;
-}
-
-UObject* FObjectMixerUtils::GetRowObject(TSharedPtr<ISceneOutlinerTreeItem> InTreeItem)
-{
-	if (const FObjectMixerEditorListRowActor* ActorRow = AsActorRow(InTreeItem))
-	{
 		return ActorRow->Actor.Get();
 	}
 
diff --git a/Engine/Plugins/Editor/ObjectMixer/ObjectMixer/Source/ObjectMixer/Private/Views/Widgets/SInlinePropertyCellWidget.cpp b/Engine/Plugins/Editor/ObjectMixer/ObjectMixer/Source/ObjectMixer/Private/Views/Widgets/SInlinePropertyCellWidget.cpp
--- a/Engine/Plugins/Editor/ObjectMixer/ObjectMixer/Source/ObjectMixer/Private/Views/Widgets/SInlinePropertyCellWidget.cpp
+++ b/Engine/Plugins/Editor/ObjectMixer/ObjectMixer/Source/ObjectMixer/Private/Views/Widgets/SInlinePropertyCellWidget.cpp
@@ -17,7 +17,8 @@
 void SInlinePropertyCellWidget::Construct(
 	const FArguments& InArgs, const FName InColumnName, const TSharedRef<ISceneOutlinerTreeItem> RowPtr)
 {
-	UObject* ObjectRef = FObjectMixerUtils::GetRowObject(RowPtr);
+	const bool bGetHybridComponent = true;
+	UObject* ObjectRef = FObjectMixerUtils::GetRowObject(RowPtr, bGetHybridComponent);
 	if (!ObjectRef || InColumnName.IsEqual(NAME_None))
 	{
 		return;
diff --git a/Engine/Plugins/Editor/ObjectMixer/ObjectMixer/Source/ObjectMixer/Public/Views/List/Modes/ObjectMixerOutlinerHierarchy.h b/Engine/Plugins/Editor/ObjectMixer/ObjectMixer/Source/ObjectMixer/Public/Views/List/Modes/ObjectMixerOutlinerHierarchy.h
--- a/Engine/Plugins/Editor/ObjectMixer/ObjectMixer/Source/ObjectMixer/Public/Views/List/Modes/ObjectMixerOutlinerHierarchy.h
+++ b/Engine/Plugins/Editor/ObjectMixer/ObjectMixer/Source/ObjectMixer/Public/Views/List/Modes/ObjectMixerOutlinerHierarchy.h
@@ -7,6 +7,7 @@
 class FObjectMixerOutlinerMode;
 class FWorldPartitionActorDesc;
 class UActorFolder;
+class ULevelInstanceSubsystem;
 class UWorldPartition;
 
 class FObjectMixerOutlinerHierarchy : public ISceneOutlinerHierarchy, public FNoncopyable
@@ -36,6 +37,9 @@ protected:
 	bool DoesWorldObjectHaveAcceptableClass(const UObject* Object) const;
 
 	void CreateFolderChild(const FFolder& Folder, UWorld* World, TArray<FSceneOutlinerTreeItemPtr>& OutItems) const;
+
+	void EvaluateActorsByComponentsAndCreateActorItems(AActor* Actor, TArray<FSceneOutlinerTreeItemPtr>& OutItems) const;
+	void ForEachActorInLevel(AActor* Actor, const ULevelInstanceSubsystem* LevelInstanceSubsystem, TArray<FSceneOutlinerTreeItemPtr>& OutItems) const;
 	
 	/** Adds all the direct and indirect children of a world to OutItems */
 	void CreateWorldChildren(UWorld* World, TArray<FSceneOutlinerTreeItemPtr>& OutItems) const;
diff --git a/Engine/Plugins/Editor/ObjectMixer/ObjectMixer/Source/ObjectMixer/Public/Views/List/ObjectMixerEditorListRowData.h b/Engine/Plugins/Editor/ObjectMixer/ObjectMixer/Source/ObjectMixer/Public/Views/List/ObjectMixerEditorListRowData.h
--- a/Engine/Plugins/Editor/ObjectMixer/ObjectMixer/Source/ObjectMixer/Public/Views/List/ObjectMixerEditorListRowData.h
+++ b/Engine/Plugins/Editor/ObjectMixer/ObjectMixer/Source/ObjectMixer/Public/Views/List/ObjectMixerEditorListRowData.h
@@ -10,8 +10,6 @@
 
 class SObjectMixerEditorList;
 
-struct FSlateBrush;
-
 /** Defines data carried by each row type. */
 struct OBJECTMIXEREDITOR_API FObjectMixerEditorListRowData
 {
@@ -40,7 +38,8 @@ struct OBJECTMIXEREDITOR_API FObjectMixerEditorListRowData
 		EPropertyValueSetFlags::Type PropertyValueSetFlags = 0;
 	};
 	
-	FObjectMixerEditorListRowData(SSceneOutliner* InSceneOutliner, const FText& InDisplayNameOverride = FText::GetEmpty())
+	FObjectMixerEditorListRowData(
+		SSceneOutliner* InSceneOutliner, const FText& InDisplayNameOverride = FText::GetEmpty())
 	: SceneOutlinerPtr(InSceneOutliner)
 	, DisplayNameOverride(InDisplayNameOverride)
 	{}
@@ -56,16 +55,13 @@ struct OBJECTMIXEREDITOR_API FObjectMixerEditorListRowData
 	UE_NODISCARD bool GetIsTreeViewItemExpanded(const TSharedRef<ISceneOutlinerTreeItem> InRow);
 	void SetIsTreeViewItemExpanded(const TSharedRef<ISceneOutlinerTreeItem> InRow, const bool bNewExpanded);
 
-	UE_NODISCARD bool GetDoesRowPassFilters() const;
-	void SetDoesRowPassFilters(const bool bPass);
-
 	UE_NODISCARD bool GetIsSelected(const TSharedRef<ISceneOutlinerTreeItem> InRow);
 	void SetIsSelected(const TSharedRef<ISceneOutlinerTreeItem> InRow, const bool bNewSelected);
 
 	UE_NODISCARD bool HasAtLeastOneChildThatIsNotSolo(
 		const TSharedRef<ISceneOutlinerTreeItem> InRow, const bool bRecursive = true) const;
 
-	UE_NODISCARD FText GetDisplayName(TSharedPtr<ISceneOutlinerTreeItem> InTreeItem, const bool bIsHybridRow = false) const;
+	UE_NODISCARD FText GetDisplayName(TSharedPtr<ISceneOutlinerTreeItem> InTreeItem) const;
 
 	UE_NODISCARD const FText& GetDisplayNameOverride() const
 	{
@@ -79,11 +75,6 @@ struct OBJECTMIXEREDITOR_API FObjectMixerEditorListRowData
 
 	UE_NODISCARD SObjectMixerEditorList* GetListView() const;
 
-	/**
-	 * Determines the style of the tree (flat list or hierarchy)
-	 */
-	EObjectMixerTreeViewMode GetTreeViewMode();
-
 	UE_NODISCARD TArray<TSharedPtr<ISceneOutlinerTreeItem>> GetSelectedTreeViewItems() const;
 	
 	void OnChangeVisibility(const FSceneOutlinerTreeItemRef TreeItem, const bool bNewVisible);
@@ -99,7 +90,21 @@ struct OBJECTMIXEREDITOR_API FObjectMixerEditorListRowData
 
 	void ClearSoloRows() const;
 
-	bool GetIsItemOrHybridChildSelected(const TSharedRef<ISceneOutlinerTreeItem> InRow);
+	bool GetIsHybridRow() const
+	{
+		return HybridComponent.IsValid();
+	}
+
+	UActorComponent* GetHybridComponent() const
+	{
+		return HybridComponent.Get();
+	}
+	
+	/** If this row represents an actor or other container and should show the data for a single child component, define it here. */
+	void SetHybridComponent(UActorComponent* InHybridComponent)
+	{
+		HybridComponent = InHybridComponent;
+	}
 
 	void PropagateChangesToSimilarSelectedRowProperties(
 		const TSharedRef<ISceneOutlinerTreeItem> InRow, const FPropertyPropagationInfo PropertyPropagationInfo);
@@ -112,6 +117,6 @@ protected:
 	FTransientEditorVisibilityRules VisibilityRules;
 
 	FText DisplayNameOverride;
-	
-	bool bDoesRowPassFilters = true;
+
+	TWeakObjectPtr<UActorComponent> HybridComponent = nullptr;
 };
diff --git a/Engine/Plugins/Editor/ObjectMixer/ObjectMixer/Source/ObjectMixer/Public/Views/List/ObjectMixerUtils.h b/Engine/Plugins/Editor/ObjectMixer/ObjectMixer/Source/ObjectMixer/Public/Views/List/ObjectMixerUtils.h
--- a/Engine/Plugins/Editor/ObjectMixer/ObjectMixer/Source/ObjectMixer/Public/Views/List/ObjectMixerUtils.h
+++ b/Engine/Plugins/Editor/ObjectMixer/ObjectMixer/Source/ObjectMixer/Public/Views/List/ObjectMixerUtils.h
@@ -27,13 +27,8 @@ namespace FObjectMixerUtils
 	FObjectMixerEditorListRowUObject* AsObjectRow(TSharedPtr<ISceneOutlinerTreeItem> InTreeItem);
 	
 	FObjectMixerEditorListRowData* GetRowData(TSharedPtr<ISceneOutlinerTreeItem> InTreeItem);
-	
-	/** If this is a hybrid row, return the child row used to hybrid with. */
-	UE_NODISCARD TWeakPtr<ISceneOutlinerTreeItem> GetHybridChild(TSharedPtr<ISceneOutlinerTreeItem> InTreeItem);
-
-	TSharedRef<ISceneOutlinerTreeItem> GetHybridChildOrRowItemIfNull(const TSharedRef<ISceneOutlinerTreeItem> InRow);
 
-	UObject* GetRowObject(TSharedPtr<ISceneOutlinerTreeItem> InTreeItem);
+	UObject* GetRowObject(TSharedPtr<ISceneOutlinerTreeItem> InTreeItem, const bool bGetHybridRowComponent = false);
 	
 	UE_NODISCARD AActor* GetSelfOrOuterAsActor(TSharedPtr<ISceneOutlinerTreeItem> InTreeItem);
 
