Object Mixer:

-Fix: Editor crashes when opening Generic Object Mixer Instance in Light Mixer panel

diff --git a/Engine/Plugins/Editor/ObjectMixer/ObjectMixer/Source/ObjectMixer/Private/Views/List/ObjectMixerEditorList.cpp b/Engine/Plugins/Editor/ObjectMixer/ObjectMixer/Source/ObjectMixer/Private/Views/List/ObjectMixerEditorList.cpp
--- a/Engine/Plugins/Editor/ObjectMixer/ObjectMixer/Source/ObjectMixer/Private/Views/List/ObjectMixerEditorList.cpp
+++ b/Engine/Plugins/Editor/ObjectMixer/ObjectMixer/Source/ObjectMixer/Private/Views/List/ObjectMixerEditorList.cpp
@@ -77,14 +77,17 @@ const UObjectMixerObjectFilter* FObjectMixerEditorList::GetMainObjectFilterInsta
 void FObjectMixerEditorList::CacheObjectFilterInstances()
 {
 	ObjectFilterInstances.Reset();
-	
-	for (const TSubclassOf<UObjectMixerObjectFilter> Class : GetObjectFilterClasses())
+
+	if (ObjectFilterClasses.Num() > 0)
 	{
-		UObjectMixerObjectFilter* NewInstance = NewObject<UObjectMixerObjectFilter>(GetTransientPackage(), Class);
-		ObjectFilterInstances.Add(TObjectPtr<UObjectMixerObjectFilter>(NewInstance));
-	}
+		for (const TSubclassOf<UObjectMixerObjectFilter> Class : ObjectFilterClasses)
+		{
+			UObjectMixerObjectFilter* NewInstance = NewObject<UObjectMixerObjectFilter>(GetTransientPackage(), Class);
+			ObjectFilterInstances.Add(TObjectPtr<UObjectMixerObjectFilter>(NewInstance));
+		}
 
-	BuildPerformanceCache();
+		BuildPerformanceCache();
+	}
 }
 
 TSet<UClass*> FObjectMixerEditorList::ForceGetObjectClassesToFilter()
