Object Mixer:

-Fix: Changing Parameters on a Multi-Selection Only Change Values Of Lights On Screen

diff --git a/Engine/Plugins/Editor/ObjectMixer/ObjectMixer/Source/ObjectMixer/Private/Views/List/ObjectMixerEditorListRowData.cpp b/Engine/Plugins/Editor/ObjectMixer/ObjectMixer/Source/ObjectMixer/Private/Views/List/ObjectMixerEditorListRowData.cpp
--- a/Engine/Plugins/Editor/ObjectMixer/ObjectMixer/Source/ObjectMixer/Private/Views/List/ObjectMixerEditorListRowData.cpp
+++ b/Engine/Plugins/Editor/ObjectMixer/ObjectMixer/Source/ObjectMixer/Private/Views/List/ObjectMixerEditorListRowData.cpp
@@ -175,6 +175,10 @@ void SetValueOnSelectedItems(
 			}
 
 			FObjectMixerEditorListRowData* RowData = FObjectMixerUtils::GetRowData(SelectedRow);
+			if (!RowData)
+			{
+				continue;
+			}
 
 			// Skip folders
 			if (FObjectMixerUtils::AsFolderRow(SelectedRow))
@@ -191,7 +195,7 @@ void SetValueOnSelectedItems(
 			else
 			{
 				UE_LOG(LogObjectMixerEditor, Warning, TEXT("%hs: Row '%s' has no valid associated object to modify."), __FUNCTION__, *RowData->GetDisplayName(SelectedRow).ToString());
-				return;
+				continue;
 			}
 		
 			// Use handles if valid, otherwise use ImportText
