Object Mixer:

-Fix: Objects with invisible components are still picked up by [REDACTED_PERSON]
-Fix: filtering lights only works with light type and not Actor name

diff --git a/Engine/Plugins/Editor/ObjectMixer/ObjectMixer/Source/ObjectMixer/Private/Views/List/Modes/ObjectMixerOutlinerHierarchy.cpp b/Engine/Plugins/Editor/ObjectMixer/ObjectMixer/Source/ObjectMixer/Private/Views/List/Modes/ObjectMixerOutlinerHierarchy.cpp
--- a/Engine/Plugins/Editor/ObjectMixer/ObjectMixer/Source/ObjectMixer/Private/Views/List/Modes/ObjectMixerOutlinerHierarchy.cpp
+++ b/Engine/Plugins/Editor/ObjectMixer/ObjectMixer/Source/ObjectMixer/Private/Views/List/Modes/ObjectMixerOutlinerHierarchy.cpp
@@ -416,28 +416,6 @@ FSceneOutlinerTreeItemPtr FObjectMixerOutlinerHierarchy::FindOrCreateParentItem(
 	}
 }
 
-void FObjectMixerOutlinerHierarchy::CreateComponentItems(
-	const AActor* Actor, TArray<FSceneOutlinerTreeItemPtr>& OutItems, const TSet<UObject*>& AllMatchingObjects) const
-{
-	check(Actor);
-	// Add all this actors components if showing components and the owning actor was created
-	if (bShowingComponents)
-	{
-		for (UActorComponent* Component : Actor->GetComponents())
-		{
-			if (Component != nullptr && AllMatchingObjects.Contains(Component))
-			{
-				if (FSceneOutlinerTreeItemPtr ComponentItem =
-					Mode->CreateItemFor<FObjectMixerEditorListRowComponent>(
-						FObjectMixerEditorListRowComponent(Component, GetCastedMode()->GetSceneOutliner())))
-				{
-					OutItems.Add(ComponentItem);
-				}
-			}
-		}
-	}
-}
-
 bool FObjectMixerOutlinerHierarchy::DoesWorldObjectHaveAcceptableClass(const UObject* Object) const
 {
 	if (IsValid(Object))
@@ -464,46 +442,6 @@ void FObjectMixerOutlinerHierarchy::CreateFolderChild(const FFolder& Folder, UWo
     }
 }
 
-void FObjectMixerOutlinerHierarchy::EvaluateActorsByComponentsAndCreateActorItems(
-	AActor* Actor, TArray<FSceneOutlinerTreeItemPtr>& OutItems) const
-{
-	if (!Actor)
-	{
-		return;
-	}
-
-	if (FSceneOutlinerTreeItemPtr ActorItem =
-			Mode->CreateItemFor<FObjectMixerEditorListRowActor>(
-				FObjectMixerEditorListRowActor(Actor, GetCastedMode()->GetSceneOutliner())))
-	{	
-		// Create all component items
-		TArray<FSceneOutlinerTreeItemPtr> ComponentItems;
-		CreateComponentItems(Actor, ComponentItems);
-
-		if (ComponentItems.Num()) // Only add items if we have components
-		{	
-			OutItems.Add(ActorItem);
-			
-			if (ComponentItems.Num() == 1) // Create hybrid row
-			{
-				if (FObjectMixerEditorListRowActor* AsActorRow = FObjectMixerUtils::AsActorRow(ActorItem))
-				{
-					const FComponentTreeItem* ComponentItem = ComponentItems[0]->CastTo<FComponentTreeItem>();
-
-					if (ComponentItem && ComponentItem->Component.IsValid())
-					{
-						AsActorRow->RowData.SetHybridComponent(ComponentItem->Component.Get());
-					}
-				}
-			}
-			else // Create normal Actor->Component hierarchy if there are multiple matching components
-			{
-				OutItems.Append(ComponentItems);
-			}
-		}
-	}
-}
-
 void FObjectMixerOutlinerHierarchy::ForEachActorInLevel(
 	AActor* Actor, const ULevelInstanceSubsystem* LevelInstanceSubsystem, TArray<FSceneOutlinerTreeItemPtr>& OutItems) const
 {
@@ -524,23 +462,7 @@ void FObjectMixerOutlinerHierarchy::ForEachActorInLevel(
 		}
 	}
 
-	// If actor class does not match filter and we only want to show actors that have valid subobjects
-	if (bShowingOnlyActorWithValidComponents && !DoesWorldObjectHaveAcceptableClass(Actor))
-	{
-		EvaluateActorsByComponentsAndCreateActorItems(Actor, OutItems);
-	}
-	else // Create actor rows anyway
-	{
-		if (FSceneOutlinerTreeItemPtr ActorItem =
-			Mode->CreateItemFor<FObjectMixerEditorListRowActor>(
-				FObjectMixerEditorListRowActor(Actor, GetCastedMode()->GetSceneOutliner())))
-		{
-			OutItems.Add(ActorItem);
-
-			// Create all component items
-			CreateComponentItems(Actor, OutItems);
-		}
-	}	
+	OutItems.Append(ConditionallyCreateActorAndComponentItems(Actor));
 }
 
 void FObjectMixerOutlinerHierarchy::CreateWorldChildren(UWorld* World, TArray<FSceneOutlinerTreeItemPtr>& OutItems) const
@@ -630,22 +552,80 @@ void FObjectMixerOutlinerHierarchy::CreateComponentItems(const AActor* Actor, TA
 	// Add all this actors components if showing components and the owning actor was created
 	if (bShowingComponents)
 	{
-		for (UActorComponent* Component : Actor->GetComponents())
+		TArray<UActorComponent*> ActorComponents;
+		Actor->GetComponents(ActorComponents);
+		for (UActorComponent* Component : ActorComponents)
+		{
+			if (IsValid(Component) && DoesWorldObjectHaveAcceptableClass(Component))
+			{					
+				if (FSceneOutlinerTreeItemPtr ComponentItem =
+					Mode->CreateItemFor<FObjectMixerEditorListRowComponent>(
+						FObjectMixerEditorListRowComponent(Component, GetCastedMode()->GetSceneOutliner())))
+				{
+					OutItems.Add(ComponentItem);
+				}
+			}
+		}
+	}
+}
+
+TArray<FSceneOutlinerTreeItemPtr> FObjectMixerOutlinerHierarchy::ConditionallyCreateActorAndComponentItems(AActor* Actor) const
+{
+	check(IsValid(Actor));
+	
+	TArray<FSceneOutlinerTreeItemPtr> ReturnValue;
+	
+	// Whether or not we have components to return, we should create an actor row if components would be returned if there were no filters
+	bool bWouldReturnAnyComponentsBeforeFiltering = false;
+	
+	if (bShowingComponents)
+	{
+		TArray<UActorComponent*> ActorComponents;
+		Actor->GetComponents(ActorComponents);
+		for (UActorComponent* Component : ActorComponents)
 		{
-			if (Component != nullptr)
+			if (IsValid(Component) && DoesWorldObjectHaveAcceptableClass(Component))
 			{
-				if (DoesWorldObjectHaveAcceptableClass(Component))
+				bWouldReturnAnyComponentsBeforeFiltering = true;
+					
+				if (FSceneOutlinerTreeItemPtr ComponentItem =
+					Mode->CreateItemFor<FObjectMixerEditorListRowComponent>(
+						FObjectMixerEditorListRowComponent(Component, GetCastedMode()->GetSceneOutliner())))
 				{
-					if (FSceneOutlinerTreeItemPtr ComponentItem =
-						Mode->CreateItemFor<FObjectMixerEditorListRowComponent>(
-							FObjectMixerEditorListRowComponent(Component, GetCastedMode()->GetSceneOutliner())))
+					ReturnValue.Add(ComponentItem);
+				}
+			}
+		}
+	}
+
+	const bool bShouldCreateActorItem =
+		bWouldReturnAnyComponentsBeforeFiltering || !bShowingOnlyActorWithValidComponents && DoesWorldObjectHaveAcceptableClass(Actor);
+	
+	if (bShouldCreateActorItem)
+	{
+		if (const FSceneOutlinerTreeItemPtr ActorItem =
+			Mode->CreateItemFor<FObjectMixerEditorListRowActor>(
+				FObjectMixerEditorListRowActor(Actor, GetCastedMode()->GetSceneOutliner())))
+		{	
+			if (ReturnValue.Num() == 1) // Create hybrid row
+			{
+				if (FObjectMixerEditorListRowActor* AsActorRow = FObjectMixerUtils::AsActorRow(ActorItem))
+				{
+					const FComponentTreeItem* ComponentItem = ReturnValue[0]->CastTo<FComponentTreeItem>();
+
+					if (ComponentItem && ComponentItem->Component.IsValid())
 					{
-						OutItems.Add(ComponentItem);
+						AsActorRow->RowData.SetHybridComponent(ComponentItem->Component.Get());
 					}
 				}
 			}
+
+			// Place ActorItem before components
+			ReturnValue.Insert(ActorItem, 0);
 		}
 	}
+
+	return ReturnValue;
 }
 
 void FObjectMixerOutlinerHierarchy::CreateItems(TArray<FSceneOutlinerTreeItemPtr>& OutItems) const
diff --git a/Engine/Plugins/Editor/ObjectMixer/ObjectMixer/Source/ObjectMixer/Public/Views/List/Modes/ObjectMixerOutlinerHierarchy.h b/Engine/Plugins/Editor/ObjectMixer/ObjectMixer/Source/ObjectMixer/Public/Views/List/Modes/ObjectMixerOutlinerHierarchy.h
--- a/Engine/Plugins/Editor/ObjectMixer/ObjectMixer/Source/ObjectMixer/Public/Views/List/Modes/ObjectMixerOutlinerHierarchy.h
+++ b/Engine/Plugins/Editor/ObjectMixer/ObjectMixer/Source/ObjectMixer/Public/Views/List/Modes/ObjectMixerOutlinerHierarchy.h
@@ -38,16 +38,21 @@ protected:
 
 	void CreateFolderChild(const FFolder& Folder, UWorld* World, TArray<FSceneOutlinerTreeItemPtr>& OutItems) const;
 
-	void EvaluateActorsByComponentsAndCreateActorItems(AActor* Actor, TArray<FSceneOutlinerTreeItemPtr>& OutItems) const;
 	void ForEachActorInLevel(AActor* Actor, const ULevelInstanceSubsystem* LevelInstanceSubsystem, TArray<FSceneOutlinerTreeItemPtr>& OutItems) const;
 	
 	/** Adds all the direct and indirect children of a world to OutItems */
 	void CreateWorldChildren(UWorld* World, TArray<FSceneOutlinerTreeItemPtr>& OutItems) const;
-	/** Create all component items for an actor if we are showing components and place them in OutItems */
+	
+	/**
+	 * Create all component items for an actor if we are showing components and place them in OutItems. Items not added if filtered out.
+	 */
 	void CreateComponentItems(const AActor* Actor, TArray<FSceneOutlinerTreeItemPtr>& OutItems) const;
-	/** Create all component items for an actor if we are showing components and place them in OutItems */
-	void CreateComponentItems(
-		const AActor* Actor, TArray<FSceneOutlinerTreeItemPtr>& OutItems, const TSet<UObject*>& AllMatchingObjects) const;
+
+	/**
+	 * Create all component items for an actor if we are showing components and place them in OutItems. Items not added if filtered out.
+	 * If components are returned or would be returned before filters, the Actor will be added at the beginning of the array.
+	 */
+	TArray<FSceneOutlinerTreeItemPtr> ConditionallyCreateActorAndComponentItems(AActor* Actor) const;
 
 	// Update the hierarchy when actor or world changing events occur
 	void OnWorldPartitionCreated(UWorld* InWorld);
