Object Mixer:

-Fix: Issue 259342: Compile errors in Module.ObjectMixerEditor.cpp and ObjectMixerOutlinerHierarchy.cpp

diff --git a/Engine/Plugins/Editor/ObjectMixer/ObjectMixer/Source/ObjectMixer/Private/Views/List/Modes/ObjectMixerOutlinerHierarchy.cpp b/Engine/Plugins/Editor/ObjectMixer/ObjectMixer/Source/ObjectMixer/Private/Views/List/Modes/ObjectMixerOutlinerHierarchy.cpp
--- a/Engine/Plugins/Editor/ObjectMixer/ObjectMixer/Source/ObjectMixer/Private/Views/List/Modes/ObjectMixerOutlinerHierarchy.cpp
+++ b/Engine/Plugins/Editor/ObjectMixer/ObjectMixer/Source/ObjectMixer/Private/Views/List/Modes/ObjectMixerOutlinerHierarchy.cpp
@@ -599,7 +599,7 @@ TArray<FSceneOutlinerTreeItemPtr> FObjectMixerOutlinerHierarchy::ConditionallyCr
 	}
 
 	const bool bShouldCreateActorItem =
-		bWouldReturnAnyComponentsBeforeFiltering || !bShowingOnlyActorWithValidComponents && DoesWorldObjectHaveAcceptableClass(Actor);
+		bWouldReturnAnyComponentsBeforeFiltering || !bShowingOnlyActorWithValidComponents || DoesWorldObjectHaveAcceptableClass(Actor);
 	
 	if (bShouldCreateActorItem)
 	{
