Object Mixer:

-Fix: Editor crashes when Light Mixer panel is opened in a Basic level

diff --git a/Engine/Plugins/Editor/ObjectMixer/ObjectMixer/Source/ObjectMixer/Private/Views/List/SObjectMixerEditorList.cpp b/Engine/Plugins/Editor/ObjectMixer/ObjectMixer/Source/ObjectMixer/Private/Views/List/SObjectMixerEditorList.cpp
--- a/Engine/Plugins/Editor/ObjectMixer/ObjectMixer/Source/ObjectMixer/Private/Views/List/SObjectMixerEditorList.cpp
+++ b/Engine/Plugins/Editor/ObjectMixer/ObjectMixer/Source/ObjectMixer/Private/Views/List/SObjectMixerEditorList.cpp
@@ -837,7 +837,19 @@ bool SObjectMixerEditorList::AddUniquePropertyColumnInfo(
 	return false;
 }
 
-void SObjectMixerEditorList::CreateActorTextInfoColumns(UWorld *WorldPtr, FSceneOutlinerInitializationOptions& OutInitOptions)
+// The "Level" column should be named "Package Short Name" in wp enabled levels
+FText SObjectMixerEditorList::GetLevelColumnName() const
+{
+	const UWorld* WorldPtr = GetWorld();
+	if (WorldPtr && WorldPtr->IsPartitionedWorld())
+	{
+		return FSceneOutlinerBuiltInColumnTypes::PackageShortName_Localized();
+	}
+
+	return FSceneOutlinerBuiltInColumnTypes::Level_Localized();
+}
+
+void SObjectMixerEditorList::CreateActorTextInfoColumns(FSceneOutlinerInitializationOptions& OutInitOptions)
 {
 	FGetTextForItem LayerInfoText = FGetTextForItem::CreateLambda([](const ISceneOutlinerTreeItem& Item) -> FString
 	{
@@ -1129,18 +1141,7 @@ void SObjectMixerEditorList::CreateActorTextInfoColumns(UWorld *WorldPtr, FScene
 		);
 	};
 
-	// The "Level" column should be named "Package Short Name" in wp enabled levels
-	auto LevelColumnName = TAttribute<FText>::CreateLambda([WorldPtr]() -> FText
-	{
-		if (WorldPtr && WorldPtr->IsPartitionedWorld())
-		{
-			return FSceneOutlinerBuiltInColumnTypes::PackageShortName_Localized();
-		}
-
-		return FSceneOutlinerBuiltInColumnTypes::Level_Localized();
-	});
-
-	AddTextInfoColumn(FSceneOutlinerBuiltInColumnTypes::Level(), LevelColumnName, LevelInfoText);
+	AddTextInfoColumn(FSceneOutlinerBuiltInColumnTypes::Level(), TAttribute<FText>::CreateSP(this, &SObjectMixerEditorList::GetLevelColumnName), LevelInfoText);
 	AddTextInfoColumn(FSceneOutlinerBuiltInColumnTypes::Layer(), FSceneOutlinerBuiltInColumnTypes::Layer_Localized(), LayerInfoText);
 	AddTextInfoColumn(FSceneOutlinerBuiltInColumnTypes::DataLayer(), FSceneOutlinerBuiltInColumnTypes::DataLayer_Localized(), DataLayerInfoText);
 	AddTextInfoColumn(FSceneOutlinerBuiltInColumnTypes::ContentBundle(), FSceneOutlinerBuiltInColumnTypes::ContentBundle_Localized(), ContentBundleInfoText);
@@ -1237,7 +1238,7 @@ void SObjectMixerEditorList::SetupColumns(FSceneOutlinerInitializationOptions& O
 
 		if (UWorld* WorldPtr = GetWorld())
 		{
-			CreateActorTextInfoColumns(WorldPtr, OutInitOptions);
+			CreateActorTextInfoColumns(OutInitOptions);
 			
 			if (WorldPtr->IsPartitionedWorld())
 			{
diff --git a/Engine/Plugins/Editor/ObjectMixer/ObjectMixer/Source/ObjectMixer/Public/Views/List/SObjectMixerEditorList.h b/Engine/Plugins/Editor/ObjectMixer/ObjectMixer/Source/ObjectMixer/Public/Views/List/SObjectMixerEditorList.h
--- a/Engine/Plugins/Editor/ObjectMixer/ObjectMixer/Source/ObjectMixer/Public/Views/List/SObjectMixerEditorList.h
+++ b/Engine/Plugins/Editor/ObjectMixer/ObjectMixer/Source/ObjectMixer/Public/Views/List/SObjectMixerEditorList.h
@@ -224,6 +224,7 @@ protected:
 		const TSet<FName>& PropertySkipList = {}
 	);
 
-	void CreateActorTextInfoColumns(UWorld *WorldPtr, FSceneOutlinerInitializationOptions& OutInitOptions);
+	FText GetLevelColumnName() const;
+	void CreateActorTextInfoColumns(FSceneOutlinerInitializationOptions& OutInitOptions);
 	void SetupColumns(FSceneOutlinerInitializationOptions& OutInitOptions);
 };
