Object Mixer:

-Fix: Hybrid Rows still able to fold out into child component
-Fix: Crash when switching to Light Mixer tab
-Fix: Editor crashes when an actor is cut and pasted
-Fix: Editor crashes when an actor is deleted or cut

diff --git a/Engine/Plugins/Editor/ObjectMixer/ObjectMixer/Source/ObjectMixer/Private/Views/List/Modes/ObjectMixerOutlinerHierarchy.cpp b/Engine/Plugins/Editor/ObjectMixer/ObjectMixer/Source/ObjectMixer/Private/Views/List/Modes/ObjectMixerOutlinerHierarchy.cpp
--- a/Engine/Plugins/Editor/ObjectMixer/ObjectMixer/Source/ObjectMixer/Private/Views/List/Modes/ObjectMixerOutlinerHierarchy.cpp
+++ b/Engine/Plugins/Editor/ObjectMixer/ObjectMixer/Source/ObjectMixer/Private/Views/List/Modes/ObjectMixerOutlinerHierarchy.cpp
@@ -571,13 +571,15 @@ void FObjectMixerOutlinerHierarchy::CreateComponentItems(const AActor* Actor, TA
 
 TArray<FSceneOutlinerTreeItemPtr> FObjectMixerOutlinerHierarchy::ConditionallyCreateActorAndComponentItems(AActor* Actor) const
 {
-	check(IsValid(Actor));
-	
-	TArray<FSceneOutlinerTreeItemPtr> ReturnValue;
+	if (!IsValid(Actor))
+	{
+		return {};
+	}
 	
 	// Whether or not we have components to return, we should create an actor row if components would be returned if there were no filters
 	bool bWouldReturnAnyComponentsBeforeFiltering = false;
 	
+	TArray<FSceneOutlinerTreeItemPtr> ComponentRows;
 	if (bShowingComponents)
 	{
 		TArray<UActorComponent*> ActorComponents;
@@ -592,7 +594,7 @@ TArray<FSceneOutlinerTreeItemPtr> FObjectMixerOutlinerHierarchy::ConditionallyCr
 					Mode->CreateItemFor<FObjectMixerEditorListRowComponent>(
 						FObjectMixerEditorListRowComponent(Component, GetCastedMode()->GetSceneOutliner())))
 				{
-					ReturnValue.Add(ComponentItem);
+					ComponentRows.Add(ComponentItem);
 				}
 			}
 		}
@@ -600,18 +602,19 @@ TArray<FSceneOutlinerTreeItemPtr> FObjectMixerOutlinerHierarchy::ConditionallyCr
 
 	const bool bShouldCreateActorItem =
 		bWouldReturnAnyComponentsBeforeFiltering || !bShowingOnlyActorWithValidComponents || DoesWorldObjectHaveAcceptableClass(Actor);
-	
+
+	TArray<FSceneOutlinerTreeItemPtr> ReturnValue;
 	if (bShouldCreateActorItem)
 	{
 		if (const FSceneOutlinerTreeItemPtr ActorItem =
 			Mode->CreateItemFor<FObjectMixerEditorListRowActor>(
 				FObjectMixerEditorListRowActor(Actor, GetCastedMode()->GetSceneOutliner())))
 		{	
-			if (ReturnValue.Num() == 1) // Create hybrid row
+			if (ComponentRows.Num() == 1) // Create hybrid row
 			{
 				if (FObjectMixerEditorListRowActor* AsActorRow = FObjectMixerUtils::AsActorRow(ActorItem))
 				{
-					const FComponentTreeItem* ComponentItem = ReturnValue[0]->CastTo<FComponentTreeItem>();
+					const FComponentTreeItem* ComponentItem = ComponentRows[0]->CastTo<FComponentTreeItem>();
 
 					if (ComponentItem && ComponentItem->Component.IsValid())
 					{
@@ -619,6 +622,10 @@ TArray<FSceneOutlinerTreeItemPtr> FObjectMixerOutlinerHierarchy::ConditionallyCr
 					}
 				}
 			}
+			else
+			{
+				ReturnValue.Append(ComponentRows);
+			}
 
 			// Place ActorItem before components
 			ReturnValue.Insert(ActorItem, 0);
