Object Mixer:

-Fix: Light Collections are not available

diff --git a/Engine/Plugins/Editor/ObjectMixer/ObjectMixer/Source/ObjectMixer/Private/Views/List/Modes/ObjectMixerOutlinerMode.cpp b/Engine/Plugins/Editor/ObjectMixer/ObjectMixer/Source/ObjectMixer/Private/Views/List/Modes/ObjectMixerOutlinerMode.cpp
--- a/Engine/Plugins/Editor/ObjectMixer/ObjectMixer/Source/ObjectMixer/Private/Views/List/Modes/ObjectMixerOutlinerMode.cpp
+++ b/Engine/Plugins/Editor/ObjectMixer/ObjectMixer/Source/ObjectMixer/Private/Views/List/Modes/ObjectMixerOutlinerMode.cpp
@@ -9,6 +9,7 @@
 #include "Views/List/Modes/ObjectMixerOutlinerHierarchy.h"
 #include "Views/List/Modes/SFilterClassMenuItem.h"
 #include "Views/List/ObjectMixerUtils.h"
+#include "Views/Widgets/ObjectMixerEditorListMenuContext.h"
 
 #include "ActorDescTreeItem.h"
 #include "ActorEditorUtils.h"
@@ -27,6 +28,7 @@
 #include "IContentBrowserSingleton.h"
 #include "ISourceControlModule.h"
 #include "LevelEditor.h"
+#include "LevelEditorContextMenu.h"
 #include "LevelTreeItem.h"
 #include "SceneOutlinerDelegates.h"
 #include "SceneOutlinerMenuContext.h"
@@ -1516,8 +1518,29 @@ TSharedPtr<SWidget> FObjectMixerOutlinerMode::BuildContextMenu()
 
 	FToolMenuContext Context(ContextObject);
 
-	FName MenuName = DefaultContextMenuName;
-	SceneOutliner->GetSharedData().ModifyContextMenu.ExecuteIfBound(MenuName, Context);
+	UObjectMixerEditorListMenuContext* ObjectMixerContextObject = NewObject<UObjectMixerEditorListMenuContext>();
+	ObjectMixerContextObject->Data = {ListModelPtr.Pin()->GetSelectedTreeViewItems(), ListModelPtr.Pin()};
+
+	Context.AddObject(ObjectMixerContextObject, [](UObject* InContext)
+	{
+		UObjectMixerEditorListMenuContext* CastContext = CastChecked<UObjectMixerEditorListMenuContext>(InContext);
+		CastContext->Data.SelectedItems.Empty();
+		CastContext->Data.ListModelPtr.Reset();
+	});
+
+	UObjectMixerEditorListMenuContext::RegisterObjectMixerDynamicCollectionsContextMenuExtension("LevelEditor.ActorContextMenu");
+	UObjectMixerEditorListMenuContext::RegisterObjectMixerDynamicCollectionsContextMenuExtension("LevelEditor.ComponentContextMenu");
+	UObjectMixerEditorListMenuContext::RegisterObjectMixerDynamicCollectionsContextMenuExtension("LevelEditor.ElementContextMenu");
+
+	const FName MenuName = "LevelEditor.LevelEditorSceneOutliner.ContextMenu";
+	
+	const FLevelEditorModule& LevelEditorModule = FModuleManager::LoadModuleChecked<FLevelEditorModule>("LevelEditor");
+	const TSharedPtr<ILevelEditor> LevelEditorPtr = LevelEditorModule.GetLevelEditorInstance().Pin();
+	
+	if (LevelEditorPtr.IsValid())
+	{
+		FLevelEditorContextMenu::InitMenuContext(Context, LevelEditorPtr, ELevelEditorMenuContext::SceneOutliner);
+	}
 
 	// Build up the menu for a selection
 	UToolMenus* ToolMenus = UToolMenus::Get();
diff --git a/Engine/Plugins/Editor/ObjectMixer/ObjectMixer/Source/ObjectMixer/Private/Views/List/SObjectMixerEditorList.cpp b/Engine/Plugins/Editor/ObjectMixer/ObjectMixer/Source/ObjectMixer/Private/Views/List/SObjectMixerEditorList.cpp
--- a/Engine/Plugins/Editor/ObjectMixer/ObjectMixer/Source/ObjectMixer/Private/Views/List/SObjectMixerEditorList.cpp
+++ b/Engine/Plugins/Editor/ObjectMixer/ObjectMixer/Source/ObjectMixer/Private/Views/List/SObjectMixerEditorList.cpp
@@ -21,8 +21,6 @@
 #include "Editor.h"
 #include "EditorActorFolders.h"
 #include "IPlacementModeModule.h"
-#include "LevelEditor.h"
-#include "LevelEditorContextMenu.h"
 #include "SceneOutlinerModule.h"
 #include "SceneOutlinerTextInfoColumn.h"
 #include "SPositiveActionButton.h"
@@ -60,33 +58,6 @@ void SObjectMixerEditorList::Construct(const FArguments& InArgs, TSharedRef<FObj
 	InitOptions.bShowParentTree = true;
 	InitOptions.OutlinerIdentifier = ListModel->GetModuleName();
 	
-	// Context Menu
-	{
-		UToolMenus* ToolMenus = UToolMenus::Get();
-		static const FName MenuName = "LevelEditor.LevelEditorSceneOutliner.ContextMenu";
-		if (!ToolMenus->IsMenuRegistered(MenuName))
-		{
-			UToolMenu* Menu = ToolMenus->RegisterMenu(MenuName, "SceneOutliner.DefaultContextMenuBase");
-			FToolMenuSection& Section = Menu->AddDynamicSection("LevelEditorContextMenu",
-				FNewToolMenuDelegate::CreateRaw(this, &SObjectMixerEditorList::OnAddDynamicContextMenuSection)
-			);
-			Section.InsertPosition = FToolMenuInsert("MainSection", EToolMenuInsertType::Before);
-		}
-
-		InitOptions.ModifyContextMenu.BindLambda([=](FName& OutMenuName, FToolMenuContext& MenuContext)
-			{
-				OutMenuName = MenuName;
-
-				const FLevelEditorModule& LevelEditorModule = FModuleManager::LoadModuleChecked<FLevelEditorModule>("LevelEditor");
-				const TSharedPtr<ILevelEditor> LevelEditorPtr = LevelEditorModule.GetLevelEditorInstance().Pin();
-
-				if (LevelEditorPtr.IsValid())
-				{
-					FLevelEditorContextMenu::InitMenuContext(MenuContext, LevelEditorPtr, ELevelEditorMenuContext::SceneOutliner);
-				}
-			});
-	}
-	
 	// Mode
 	InitOptions.ModeFactory = FCreateSceneOutlinerMode::CreateLambda([this](SSceneOutliner* InOutliner)
 	{		
@@ -139,33 +110,6 @@ bool SObjectMixerEditorList::CanCreateFolder() const
 	return GetSelectedTreeViewItemCount() > 0;
 }
 
-void SObjectMixerEditorList::OnAddDynamicContextMenuSection(UToolMenu* InMenu)
-{
-	FName LevelContextMenuName = "LevelEditor.ActorContextMenu";
-
-	const FSceneOutlinerItemSelection Selection = GetSelection();
-	
-	if (Selection.Has<FComponentTreeItem>())
-	{
-		LevelContextMenuName = "LevelEditor.ComponentContextMenu";
-	}
-	else if (Selection.Has<FActorTreeItem>())
-	{
-		LevelContextMenuName = "LevelEditor.ActorContextMenu";
-	}
-	else if (Selection.Num() > 0)
-	{
-		LevelContextMenuName = "LevelEditor.ElementContextMenu";
-	}
-				
-	if (LevelContextMenuName != NAME_None)
-	{
-		// Extend the menu even if no actors selected, as Edit menu should always exist for scene outliner
-		UToolMenu* OtherMenu = UToolMenus::Get()->GenerateMenu(LevelContextMenuName, InMenu->Context);
-		InMenu->Sections.Append(OtherMenu->Sections);
-	}
-}
-
 TSharedRef<SWidget> SObjectMixerEditorList::OnGenerateAddObjectButtonMenu() const
 {
 	TSet<UClass*> ClassesToPlace;
diff --git a/Engine/Plugins/Editor/ObjectMixer/ObjectMixer/Source/ObjectMixer/Private/Views/Widgets/ObjectMixerEditorListMenuContext.cpp b/Engine/Plugins/Editor/ObjectMixer/ObjectMixer/Source/ObjectMixer/Private/Views/Widgets/ObjectMixerEditorListMenuContext.cpp
--- a/Engine/Plugins/Editor/ObjectMixer/ObjectMixer/Source/ObjectMixer/Private/Views/Widgets/ObjectMixerEditorListMenuContext.cpp
+++ b/Engine/Plugins/Editor/ObjectMixer/ObjectMixer/Source/ObjectMixer/Private/Views/Widgets/ObjectMixerEditorListMenuContext.cpp
@@ -4,23 +4,14 @@
 
 #include "Views/List/ObjectMixerEditorList.h"
 #include "Views/List/ObjectMixerUtils.h"
-#include "Views/List/RowTypes/ObjectMixerEditorListRowFolder.h"
-
-#include "ActorFolderPickingMode.h"
-#include "FolderTreeItem.h"
-#include "LevelEditor.h"
-#include "LevelEditorContextMenu.h"
-#include "LevelEditorMenuContext.h"
-#include "ObjectMixerEditorModule.h"
-#include "Selection.h"
-#include "SSceneOutliner.h"
+
+#include "Algo/Find.h"
 #include "ToolMenu.h"
 #include "ToolMenuContext.h"
 #include "ToolMenuDelegates.h"
 #include "ToolMenuEntry.h"
 #include "ToolMenus.h"
 #include "ToolMenuSection.h"
-#include "Elements/Framework/TypedElementSelectionSet.h"
 #include "Framework/Application/SlateApplication.h"
 #include "Framework/Commands/GenericCommands.h"
 #include "Framework/MultiBox/MultiBoxDefs.h"
@@ -30,18 +21,6 @@
 
 #define LOCTEXT_NAMESPACE "ObjectMixerEditor"
 
-FName UObjectMixerEditorListMenuContext::DefaultContextBaseMenuName("ObjectMixer.ContextMenuBase");
-
-TSharedPtr<SWidget> UObjectMixerEditorListMenuContext::CreateContextMenu(FObjectMixerEditorListMenuContextData InData)
-{
-	if (InData.SelectedItems.Num() == 0)
-	{
-		return nullptr;
-	}
-
-	return BuildContextMenu(MoveTemp(InData));
-}
-
 bool UObjectMixerEditorListMenuContext::DoesSelectionHaveType(const FObjectMixerEditorListMenuContextData& InData, UClass* Type)
 {
 	if (const TSharedPtr<ISceneOutlinerTreeItem>* Match = Algo::FindByPredicate(
@@ -58,125 +37,6 @@ bool UObjectMixerEditorListMenuContext::DoesSelectionHaveType(const FObjectMixer
 	return false;
 }
 
-void PerformLevelEditorRegistrations(FToolMenuContext& Context)
-{
-	FLevelEditorModule& LevelEditorModule = FModuleManager::LoadModuleChecked<FLevelEditorModule>("LevelEditor");
-	TSharedPtr<ILevelEditor> LevelEditorPtr = LevelEditorModule.GetLevelEditorInstance().Pin();
-	check(LevelEditorPtr);
-
-	TSharedPtr<FUICommandList> LevelEditorActionsList = LevelEditorPtr->GetLevelEditorActions();
-	Context.AppendCommandList(LevelEditorActionsList);
-
-	ULevelEditorContextMenuContext* LevelEditorContextObject = NewObject<ULevelEditorContextMenuContext>();
-	LevelEditorContextObject->LevelEditor = LevelEditorPtr;
-	LevelEditorContextObject->ContextType = ELevelEditorMenuContext::SceneOutliner;
-	LevelEditorContextObject->CurrentSelection = LevelEditorPtr->GetElementSelectionSet();
-
-	for (FSelectedEditableComponentIterator It(GEditor->GetSelectedEditableComponentIterator()); It; ++It)
-	{
-		LevelEditorContextObject->SelectedComponents.Add(CastChecked<UActorComponent>(*It));
-	}
-
-	Context.AddObject(LevelEditorContextObject, [](UObject* InContext)
-	{
-		ULevelEditorContextMenuContext* CastContext = CastChecked<ULevelEditorContextMenuContext>(InContext);
-		CastContext->CurrentSelection = nullptr;
-		CastContext->HitProxyElement.Release();
-	});
-
-	if (LevelEditorPtr->GetElementSelectionSet()->GetSelectedObjects<UActorComponent>().Num() == 0 &&
-		LevelEditorPtr->GetElementSelectionSet()->GetSelectedObjects<AActor>().Num() > 0)
-	{
-		const TArray<AActor*> SelectedActors = LevelEditorPtr->GetElementSelectionSet()->GetSelectedObjects<AActor>();
-
-		// Get all menu extenders for this context menu from the level editor module
-		TArray<FLevelEditorModule::FLevelViewportMenuExtender_SelectedActors> MenuExtenderDelegates =
-			LevelEditorModule.GetAllLevelViewportContextMenuExtenders();
-		TArray<TSharedPtr<FExtender>> Extenders;
-		for (int32 i = 0; i < MenuExtenderDelegates.Num(); ++i)
-		{
-			if (MenuExtenderDelegates[i].IsBound())
-			{
-				Extenders.Add(MenuExtenderDelegates[i].Execute(LevelEditorActionsList.ToSharedRef(), SelectedActors));
-			}
-		}
-
-		if (Extenders.Num() > 0)
-		{
-			Context.AddExtender(FExtender::Combine(Extenders));
-		}
-	}
-}
-
-TSharedPtr<SWidget> UObjectMixerEditorListMenuContext::BuildContextMenu(const FObjectMixerEditorListMenuContextData InData)
-{	
-	FToolMenuContext Context;
-
-	UObjectMixerEditorListMenuContext* ObjectMixerContextObject = NewObject<UObjectMixerEditorListMenuContext>();
-	ObjectMixerContextObject->Data = InData;
-
-	Context.AddObject(ObjectMixerContextObject, [](UObject* InContext)
-	{
-		UObjectMixerEditorListMenuContext* CastContext = CastChecked<UObjectMixerEditorListMenuContext>(InContext);
-		CastContext->Data.SelectedItems.Empty();
-		CastContext->Data.ListModelPtr.Reset();
-	});
-
-	if (DoesSelectionHaveType(InData, AActor::StaticClass()))
-	{
-		FLevelEditorContextMenu::RegisterActorContextMenu();
-		FLevelEditorContextMenu::RegisterElementContextMenu();
-		RegisterObjectMixerActorContextMenuExtension();
-
-		PerformLevelEditorRegistrations(Context);
-	
-		return UToolMenus::Get()->GenerateWidget("LevelEditor.ActorContextMenu", Context);
-	}
-	
-	if (DoesSelectionHaveType(InData, UActorComponent::StaticClass()))
-	{
-		FLevelEditorContextMenu::RegisterComponentContextMenu();
-		FLevelEditorContextMenu::RegisterElementContextMenu();
-		RegisterObjectMixerElementContextMenuExtension("LevelEditor.ComponentContextMenu");
-
-		PerformLevelEditorRegistrations(Context);
-
-		return UToolMenus::Get()->GenerateWidget("LevelEditor.ComponentContextMenu", Context);
-	}
-
-	if (DoesSelectionHaveType(InData, UObject::StaticClass()))
-	{
-		FLevelEditorContextMenu::RegisterElementContextMenu();
-		RegisterObjectMixerElementContextMenuExtension("LevelEditor.ElementContextMenu");
-
-		PerformLevelEditorRegistrations(Context);
-
-		return UToolMenus::Get()->GenerateWidget("LevelEditor.ElementContextMenu", Context);
-	}
-
-	// Folders only
-	RegisterFoldersOnlyContextMenu();
-
-	return UToolMenus::Get()->GenerateWidget("ObjectMixer.FoldersOnlyContextMenu", Context);
-}
-
-void UObjectMixerEditorListMenuContext::FillSelectionSubMenu(UToolMenu* Menu, const FObjectMixerEditorListMenuContextData& ContextData)
-{
-	FToolMenuSection& Section = Menu->AddSection("Section");
-	Section.AddMenuEntry(
-		"AddChildrenToSelection",
-		LOCTEXT( "AddChildrenToSelection", "Immediate Children" ),
-		LOCTEXT( "AddChildrenToSelection_ToolTip", "Select all immediate children of the selected folders" ),
-		FSlateIcon(),
-		FExecuteAction::CreateStatic(&UObjectMixerEditorListMenuContext::SelectDescendentsOfSelectedFolders, ContextData, false));
-	Section.AddMenuEntry(
-		"AddDescendantsToSelection",
-		LOCTEXT( "AddDescendantsToSelection", "All Descendants" ),
-		LOCTEXT( "AddDescendantsToSelection_ToolTip", "Select all descendants of the selected folders" ),
-		FSlateIcon(),
-		FExecuteAction::CreateStatic(&UObjectMixerEditorListMenuContext::SelectDescendentsOfSelectedFolders, ContextData, true));
-}
-
 void UObjectMixerEditorListMenuContext::RegisterFoldersOnlyContextMenu()
 {
 	UToolMenus* ToolMenus = UToolMenus::Get();
@@ -197,21 +57,6 @@ void UObjectMixerEditorListMenuContext::RegisterFoldersOnlyContextMenu()
 
 		const FObjectMixerEditorListMenuContextData& ContextData = Context->Data;
 
-		{
-			FToolMenuSection& Section = InMenu->AddSection("Hierarchy", LOCTEXT("HierarchyMenuHeader", "Hierarchy"));
-
-			Section.AddSubMenu(
-				"SelectSubMenu",
-				LOCTEXT("SelectSubMenu", "Select"),
-				FText::GetEmpty(),
-				FNewToolMenuDelegate::CreateLambda(
-					[ContextData](UToolMenu* InMenu)
-				{
-					FillSelectionSubMenu(InMenu, ContextData);
-				})
-			);
-		}
-
 		// For future CL
 		// {
 		// 	FToolMenuSection& Section = InMenu->AddSection("ElementEditActions");
@@ -227,8 +72,6 @@ void UObjectMixerEditorListMenuContext::RegisterFoldersOnlyContextMenu()
 		// 		false, // default value
 		// 		FSlateIcon(FAppStyle::GetAppStyleSetName(), "Icons.Edit"));
 		// }
-
-		GenerateMoveToMenu(InMenu, FToolMenuInsert("Hierarchy", EToolMenuInsertType::After), ContextData);
 	}));
 }
 
@@ -245,38 +88,11 @@ void UObjectMixerEditorListMenuContext::AddCollectionsMenuItem(UToolMenu* InMenu
 	);
 }
 
-void UObjectMixerEditorListMenuContext::RegisterObjectMixerActorContextMenuExtension()
-{
-	if (UToolMenu* ActorContextMenu = UToolMenus::Get()->ExtendMenu("LevelEditor.ActorContextMenu"))
-	{		
-		ActorContextMenu->AddDynamicSection("DynamicActorSection", FNewToolMenuDelegate::CreateLambda(
-			[](UToolMenu* InMenu)
-			{
-				// Ensure proper context
-				UObjectMixerEditorListMenuContext* Context = InMenu->FindContext<UObjectMixerEditorListMenuContext>();
-				if (!Context || Context->Data.SelectedItems.Num() == 0)
-				{
-					return;
-				}
-
-				const FObjectMixerEditorListMenuContextData& ContextData = Context->Data;
-
-				ReplaceEditSubMenu(ContextData);
-
-				AddCollectionsMenuItem(InMenu, ContextData);
-
-				GenerateMoveToMenu(InMenu, FToolMenuInsert("ActorTypeTools", EToolMenuInsertType::After), ContextData);
-			}),
-			FToolMenuInsert(NAME_None,EToolMenuInsertType::First)
-		);
-	}
-}
-
-void UObjectMixerEditorListMenuContext::RegisterObjectMixerElementContextMenuExtension(const FName& MenuName)
+void UObjectMixerEditorListMenuContext::RegisterObjectMixerDynamicCollectionsContextMenuExtension(const FName& MenuName)
 {
 	if (UToolMenu* ActorContextMenu = UToolMenus::Get()->ExtendMenu(MenuName))
 	{		
-		ActorContextMenu->AddDynamicSection("DynamicSubobjectSection", FNewToolMenuDelegate::CreateLambda(
+		ActorContextMenu->AddDynamicSection("DynamicCollectionsSection", FNewToolMenuDelegate::CreateLambda(
 			[](UToolMenu* InMenu)
 			{
 				// Ensure proper context
@@ -365,159 +181,6 @@ void UObjectMixerEditorListMenuContext::CreateSelectCollectionsSubMenu(UToolMenu
 	}
 }
 
-void UObjectMixerEditorListMenuContext::GenerateMoveToMenu(
-	UToolMenu* InMenu, const FToolMenuInsert& InsertArgs, const FObjectMixerEditorListMenuContextData& ContextData)
-{
-	FToolMenuSection& Section =
-		InMenu->AddSection("MainSection",
-			LOCTEXT("OutlinerSectionName", "Outliner"), InsertArgs);
-					
-	Section.AddSubMenu(
-			"MoveActorsTo",
-			LOCTEXT("MoveActorsTo", "Move To"),
-			LOCTEXT("MoveActorsTo_Tooltip", "Move selection to another folder"),
-			FNewToolMenuDelegate::CreateStatic(&UObjectMixerEditorListMenuContext::FillFoldersSubMenu, ContextData));
-}
-
-void UObjectMixerEditorListMenuContext::OnFoldersMenuFolderSelected(
-	TSharedRef<ISceneOutlinerTreeItem> Item,
-	FObjectMixerEditorListMenuContextData ContextData)
-{
-	if (const FFolderTreeItem* FolderTreeItem = Item->CastTo<FFolderTreeItem>())
-	{
-		for (const TSharedPtr<ISceneOutlinerTreeItem>& SelectedItem : ContextData.SelectedItems)
-		{
-			if (AActor* AsActor = Cast<AActor>(FObjectMixerUtils::GetRowObject(SelectedItem)))
-			{
-				AsActor->SetFolderPath_Recursively(FolderTreeItem->GetFolder().GetPath());
-			}
-			else if (FObjectMixerEditorListRowFolder* AsFolder = FObjectMixerUtils::AsFolderRow(SelectedItem))
-			{
-				if (TSharedPtr<FObjectMixerEditorList> PinnedList = ContextData.ListModelPtr.Pin())
-				{
-					//	PinnedList->OnRequestMoveFolder(AsFolder->GetFolder(), FolderTreeItem->GetFolder());
-				}
-			}
-		}
-	}
-
-	FSlateApplication::Get().DismissAllMenus();
-}
-
-TSharedRef<TSet<FFolder>> UObjectMixerEditorListMenuContext::GatherInvalidMoveToDestinations(
-	const FObjectMixerEditorListMenuContextData& ContextData)
-{
-	// We use a pointer here to save copying the whole array for every invocation of the filter delegate
-	TSharedRef<TSet<FFolder>> Exclusions(new TSet<FFolder>());
-
-	for (const TSharedPtr<ISceneOutlinerTreeItem>& SelectedItem : ContextData.SelectedItems)
-	{
-		if (FObjectMixerEditorListRowFolder* AsFolder = FObjectMixerUtils::AsFolderRow(SelectedItem))
-		{
-			Exclusions->Add(AsFolder->GetFolder());
-
-			if (const TSharedPtr<ISceneOutlinerTreeItem>& ParentRow = SelectedItem->GetParent())
-			{
-				auto FolderHasOtherSubFolders = [&SelectedItem](const TSharedPtr<ISceneOutlinerTreeItem>& InItem)
-				{
-					if (InItem != SelectedItem && FObjectMixerUtils::AsFolderRow(InItem))
-					{
-						return true;
-					}
-					return false;
-				};
-
-				// Exclude this items direct parent if it is a folder and has no other subfolders we can move to
-				bool bFolderHasOtherSubFolders = false;
-				for (const TWeakPtr<ISceneOutlinerTreeItem>& ChildRow : ParentRow->GetChildren())
-				{
-					if (ChildRow.IsValid() && FolderHasOtherSubFolders(ChildRow.Pin()))
-					{
-						bFolderHasOtherSubFolders = true;
-						break;
-					}
-				}
-	
-				if (!bFolderHasOtherSubFolders)
-				{
-					Exclusions->Add(AsFolder->GetFolder());
-				}
-			}
-		}
-	}
-
-	return Exclusions;
-}
-
-void UObjectMixerEditorListMenuContext::FillFoldersSubMenu(UToolMenu* InMenu, FObjectMixerEditorListMenuContextData ContextData)
-{
-	FToolMenuSection& Section = InMenu->AddSection("Section");
-	Section.AddMenuEntry(
-		"CreateNew",
-		LOCTEXT( "CreateNew", "Create New Folder" ),
-		LOCTEXT( "CreateNew_ToolTip", "Move to a new folder" ),
-		FSlateIcon(FAppStyle::GetAppStyleSetName(), "SceneOutliner.NewFolderIcon"),
-		FExecuteAction::CreateLambda([ContextData]()
-		{
-			if (const TSharedPtr<FObjectMixerEditorList> PinnedList = ContextData.ListModelPtr.Pin())
-			{
-				//PinnedList->OnRequestNewFolder();
-			}
-		})
-	);
-	
-	FSceneOutlinerInitializationOptions MiniSceneOutlinerInitOptions;
-	MiniSceneOutlinerInitOptions.bShowHeaderRow = false;
-	MiniSceneOutlinerInitOptions.bFocusSearchBoxWhenOpened = true;
-
-	MiniSceneOutlinerInitOptions.ModeFactory = FCreateSceneOutlinerMode::CreateLambda(
-		[ContextData](SSceneOutliner* Outliner)
-			{
-				FOnSceneOutlinerItemPicked OnSceneOutlinerItemPicked =
-					FOnSceneOutlinerItemPicked::CreateStatic(
-						&UObjectMixerEditorListMenuContext::OnFoldersMenuFolderSelected, ContextData
-				);
-
-				UWorld* World = FObjectMixerEditorModule::Get().GetWorld();
-			
-				return new FActorFolderPickingMode(Outliner, OnSceneOutlinerItemPicked, World);
-			}
-	);
-
-	TSharedRef<TSet<FFolder>> Exclusions = GatherInvalidMoveToDestinations(ContextData);
-	MiniSceneOutlinerInitOptions.Filters->AddFilterPredicate<FFolderTreeItem>(
-		FFolderTreeItem::FFilterPredicate::CreateLambda(
-			[&Exclusions](const FFolder& Folder)
-			{
-				for (const FFolder& Parent : *Exclusions)
-				{
-					if (Folder == Parent || Folder.IsChildOf(Parent))
-					{
-						return false;
-					}
-				}
-				return true;
-			}), FSceneOutlinerFilter::EDefaultBehaviour::Pass);
-
-	TSharedRef< SWidget > MiniSceneOutliner =
-		SNew(SVerticalBox)
-		+ SVerticalBox::Slot()
-		.MaxHeight(400.0f)
-		[
-			SNew(SSceneOutliner, MiniSceneOutlinerInitOptions)
-			.IsEnabled(FSlateApplication::Get().GetNormalExecutionAttribute())
-		];
-
-	FToolMenuSection& MoveToSection = InMenu->AddSection(FName(), LOCTEXT("ExistingFolders", "Existing:"));
-	
-	MoveToSection.AddEntry(FToolMenuEntry::InitWidget(
-		"MiniSceneOutliner",
-		MiniSceneOutliner,
-		FText::GetEmpty(),
-		false)
-	);
-}
-
 void UObjectMixerEditorListMenuContext::OnTextCommitted(const FText& InText, ETextCommit::Type InCommitType, const FObjectMixerEditorListMenuContextData ContextData)
 {
 	if (InCommitType == ETextCommit::OnEnter)
@@ -526,21 +189,6 @@ void UObjectMixerEditorListMenuContext::OnTextCommitted(const FText& InText, ETe
 	}
 }
 
-void UObjectMixerEditorListMenuContext::SelectDescendentsOfSelectedFolders(
-	FObjectMixerEditorListMenuContextData ContextData, const bool bRecursive)
-{
-	for (const TSharedPtr<ISceneOutlinerTreeItem>& SelectedRow : ContextData.SelectedItems)
-	{
-		if (FObjectMixerEditorListRowData* RowData = FObjectMixerUtils::GetRowData(SelectedRow))
-		{
-			if (FObjectMixerUtils::AsFolderRow(SelectedRow))
-			{
-				// RowData->SetChildRowsSelected(true, bRecursive);
-			}
-		}
-	}
-}
-
 void UObjectMixerEditorListMenuContext::OnCollectionMenuEntryCheckStateChanged(const FName Key, const FObjectMixerEditorListMenuContextData ContextData)
 {
 	if (AreAllObjectsInCollection(Key, ContextData))
diff --git a/Engine/Plugins/Editor/ObjectMixer/ObjectMixer/Source/ObjectMixer/Public/Views/List/SObjectMixerEditorList.h b/Engine/Plugins/Editor/ObjectMixer/ObjectMixer/Source/ObjectMixer/Public/Views/List/SObjectMixerEditorList.h
--- a/Engine/Plugins/Editor/ObjectMixer/ObjectMixer/Source/ObjectMixer/Public/Views/List/SObjectMixerEditorList.h
+++ b/Engine/Plugins/Editor/ObjectMixer/ObjectMixer/Source/ObjectMixer/Public/Views/List/SObjectMixerEditorList.h
@@ -84,8 +84,6 @@ public:
 
 	virtual FReply OnKeyDown( const FGeometry& MyGeometry, const FKeyEvent& InKeyEvent ) override;
 
-	void OnAddDynamicContextMenuSection(UToolMenu* InMenu);
-
 	/**
 	 * Determines the style of the tree (flat list or hierarchy)
 	 */
diff --git a/Engine/Plugins/Editor/ObjectMixer/ObjectMixer/Source/ObjectMixer/Public/Views/Widgets/ObjectMixerEditorListMenuContext.h b/Engine/Plugins/Editor/ObjectMixer/ObjectMixer/Source/ObjectMixer/Public/Views/Widgets/ObjectMixerEditorListMenuContext.h
--- a/Engine/Plugins/Editor/ObjectMixer/ObjectMixer/Source/ObjectMixer/Public/Views/Widgets/ObjectMixerEditorListMenuContext.h
+++ b/Engine/Plugins/Editor/ObjectMixer/ObjectMixer/Source/ObjectMixer/Public/Views/Widgets/ObjectMixerEditorListMenuContext.h
@@ -25,15 +25,10 @@ public:
 		TWeakPtr<class FObjectMixerEditorList> ListModelPtr;
 	};
 
-	static TSharedPtr<SWidget> CreateContextMenu(FObjectMixerEditorListMenuContextData InData);
-	static TSharedPtr<SWidget> BuildContextMenu(const FObjectMixerEditorListMenuContextData InData);
 	static void RegisterFoldersOnlyContextMenu();
-	static void RegisterObjectMixerActorContextMenuExtension();
-	static void RegisterObjectMixerElementContextMenuExtension(const FName& MenuName);
+	static void RegisterObjectMixerDynamicCollectionsContextMenuExtension(const FName& MenuName);
 
 	FObjectMixerEditorListMenuContextData Data;
-	
-	static FName DefaultContextBaseMenuName;
 
 private:
 
@@ -41,14 +36,7 @@ private:
 	static void AddCollectionWidget(const FName& Key, const FObjectMixerEditorListMenuContextData& ContextData, UToolMenu* Menu);
 	static void CreateSelectCollectionsSubMenu(UToolMenu* Menu, FObjectMixerEditorListMenuContextData ContextData);
 	static void AddCollectionsMenuItem(UToolMenu* InMenu, const FObjectMixerEditorListMenuContextData& ContextData);
-	static void GenerateMoveToMenu(UToolMenu* InMenu, const FToolMenuInsert& InsertArgs, const FObjectMixerEditorListMenuContextData& ContextData);
-	static void OnFoldersMenuFolderSelected(TSharedRef<ISceneOutlinerTreeItem> Item, FObjectMixerEditorListMenuContextData ContextData);
-	static TSharedRef<TSet<FFolder>> GatherInvalidMoveToDestinations(const FObjectMixerEditorListMenuContextData& ContextData);
-	static void FillFoldersSubMenu(UToolMenu* InMenu, FObjectMixerEditorListMenuContextData ContextData);
-	static void FillSelectionSubMenu(UToolMenu* Menu, const FObjectMixerEditorListMenuContextData& ContextData);
 	static void OnTextCommitted(const FText& InText, ETextCommit::Type InCommitType, const FObjectMixerEditorListMenuContextData ContextData);
-
-	static void SelectDescendentsOfSelectedFolders(FObjectMixerEditorListMenuContextData ContextData, const bool bRecursive);
 	
 	static void OnCollectionMenuEntryCheckStateChanged(const FName Key, const FObjectMixerEditorListMenuContextData ContextData);
 	static void AddObjectsToCollection(const FName Key, const FObjectMixerEditorListMenuContextData ContextData);
