Object Mixer:

-Fix: Collections duplicate each time the mixer panel is opened

diff --git a/Engine/Plugins/Editor/ObjectMixer/ObjectMixer/Source/ObjectMixer/Private/Views/List/ObjectMixerEditorList.cpp b/Engine/Plugins/Editor/ObjectMixer/ObjectMixer/Source/ObjectMixer/Private/Views/List/ObjectMixerEditorList.cpp
--- a/Engine/Plugins/Editor/ObjectMixer/ObjectMixer/Source/ObjectMixer/Private/Views/List/ObjectMixerEditorList.cpp
+++ b/Engine/Plugins/Editor/ObjectMixer/ObjectMixer/Source/ObjectMixer/Private/Views/List/ObjectMixerEditorList.cpp
@@ -117,7 +117,7 @@ void FObjectMixerEditorList::AddObjectFilterClass(UClass* InObjectFilterClass, c
 {
 	if (ensureAlwaysMsgf(InObjectFilterClass->IsChildOf(UObjectMixerObjectFilter::StaticClass()), TEXT("%hs: Class '%s' is not a child of UObjectMixerObjectFilter."), __FUNCTION__, *InObjectFilterClass->GetName()))
 	{
-		ObjectFilterClasses.Add(InObjectFilterClass);
+		ObjectFilterClasses.AddUnique(InObjectFilterClass);
 
 		CacheObjectFilterInstances();
 
diff --git a/Engine/Plugins/Editor/ObjectMixer/ObjectMixer/Source/ObjectMixer/Public/Views/List/ObjectMixerEditorList.h b/Engine/Plugins/Editor/ObjectMixer/ObjectMixer/Source/ObjectMixer/Public/Views/List/ObjectMixerEditorList.h
--- a/Engine/Plugins/Editor/ObjectMixer/ObjectMixer/Source/ObjectMixer/Public/Views/List/ObjectMixerEditorList.h
+++ b/Engine/Plugins/Editor/ObjectMixer/ObjectMixer/Source/ObjectMixer/Public/Views/List/ObjectMixerEditorList.h
@@ -188,7 +188,8 @@ protected:
 	TArray<TObjectPtr<UObjectMixerObjectFilter>> ObjectFilterInstances;
 
 	/**
-	 * The classes used to generate property edit columns
+	 * The classes used to generate property edit columns.
+	 * Using an array rather than a set because the first class is considered the 'Main' class which determines some filter behaviours.
 	 */
 	TArray<TSubclassOf<UObjectMixerObjectFilter>> ObjectFilterClasses;
 
