Object Mixer:

-Fix: Bulk property changes undo transactions don't work
-Fix: Changing Parameters on a Multi-Selection Only Change Values Of Lights On Screen

diff --git a/Engine/Plugins/Editor/ObjectMixer/ObjectMixer/Source/ObjectMixer/Private/Views/List/ObjectMixerEditorListRowData.cpp b/Engine/Plugins/Editor/ObjectMixer/ObjectMixer/Source/ObjectMixer/Private/Views/List/ObjectMixerEditorListRowData.cpp
--- a/Engine/Plugins/Editor/ObjectMixer/ObjectMixer/Source/ObjectMixer/Private/Views/List/ObjectMixerEditorListRowData.cpp
+++ b/Engine/Plugins/Editor/ObjectMixer/ObjectMixer/Source/ObjectMixer/Private/Views/List/ObjectMixerEditorListRowData.cpp
@@ -6,6 +6,7 @@
 #include "Views/List/ObjectMixerUtils.h"
 #include "Views/List/SObjectMixerEditorList.h"
 
+#include "Editor.h"
 #include "GameFramework/Actor.h"
 #include "ScopedTransaction.h"
 
@@ -160,12 +161,18 @@ void FObjectMixerEditorListRowData::ClearSoloRows() const
 void SetValueOnSelectedItems(
 	const FString& ValueAsString, const TArray<TSharedPtr<ISceneOutlinerTreeItem>>& OtherSelectedItems,
 	const FName& PropertyName, const TSharedPtr<ISceneOutlinerTreeItem> PinnedItem,
-	const EPropertyValueSetFlags::Type Flags)
+	const EPropertyValueSetFlags::Type InFlags)
 {
 	if (!ValueAsString.IsEmpty())
 	{
-		FScopedTransaction Transaction(
-			NSLOCTEXT("ObjectMixerEditor","OnPropertyChangedTransaction", "Object Mixer - Bulk Edit Selected Row Properties") );
+		const EPropertyValueSetFlags::Type Flags = InFlags | EPropertyValueSetFlags::NotTransactable;
+		const bool bShouldTransact = InFlags & EPropertyValueSetFlags::DefaultFlags;
+
+		if (bShouldTransact && GEditor->CanTransact())
+		{
+			GEditor->BeginTransaction(
+			   NSLOCTEXT("ObjectMixerEditor","OnPropertyChangedTransaction", "Object Mixer - Edit Selected Row Properties"));
+		}
 		
 		for (const TSharedPtr<ISceneOutlinerTreeItem>& SelectedRow : OtherSelectedItems)
 		{
@@ -186,7 +193,7 @@ void SetValueOnSelectedItems(
 				continue;
 			}
 
-			UObject* ObjectToModify = FObjectMixerUtils::GetRowObject(SelectedRow);
+			UObject* ObjectToModify = FObjectMixerUtils::GetRowObject(SelectedRow, true);
 			
 			if (IsValid(ObjectToModify))
 			{
@@ -198,26 +205,62 @@ void SetValueOnSelectedItems(
 				continue;
 			}
 		
-			// Use handles if valid, otherwise use ImportText
+			// Use handles if valid, otherwise use ImportText. Need to use the handles to ensure the Blueprints update properly.
+			
+			// Transactions are handled automatically by the handles, so no need to start a new transaction.
 			if (const TWeakPtr<IPropertyHandle>* SelectedHandlePtr = RowData->PropertyNamesToHandles.Find(PropertyName))
 			{
 				if (SelectedHandlePtr->IsValid())
 				{				
 					SelectedHandlePtr->Pin()->SetValueFromFormattedString(ValueAsString, Flags);
+					continue;
 				}
 			}
-			else
+
+			// Handles approach failed, so use ImportText
+			if (FProperty* PropertyToChange = FindFProperty<FProperty>(ObjectToModify->GetClass(), PropertyName))
 			{
-				if (const FProperty* PropertyToChange = FindFProperty<FProperty>(ObjectToModify->GetClass(), PropertyName))
+				if (void* ValuePtr = PropertyToChange->ContainerPtrToValuePtr<void>(ObjectToModify))
 				{
-					if (void* ValuePtr = PropertyToChange->ContainerPtrToValuePtr<void>(ObjectToModify))
+					if (bShouldTransact)
 					{
-						// Set the actual property value
-						PropertyToChange->ImportText_Direct(*ValueAsString, ValuePtr, ObjectToModify, PPF_None);
+						ObjectToModify->Modify();
+					}
+
+					// Set the actual property value
+					EPropertyChangeType::Type ChangeType =
+						InFlags == EPropertyValueSetFlags::InteractiveChange
+							? EPropertyChangeType::Interactive
+							: EPropertyChangeType::ValueSet;
+
+					// Set the actual property value
+					PropertyToChange->ImportText_Direct(*ValueAsString, ValuePtr, ObjectToModify, PPF_None);
+					FPropertyChangedEvent ChangeEvent(
+						PropertyToChange,
+						ChangeType,
+						MakeArrayView({ObjectToModify}));
+					ObjectToModify->PostEditChangeProperty(ChangeEvent);
+
+					// Propagate to outers
+					UObject* Outer = ObjectToModify->GetOuter();
+					while (Outer) 
+					{
+						FPropertyChangedEvent ActorChangeEvent(
+							PropertyToChange,
+							ChangeType,
+							MakeArrayView({Outer}));
+						Outer->PostEditChangeProperty(ActorChangeEvent);
+
+						Outer = Outer->GetOuter();
 					}
 				}
 			}
 		}
+		
+		if (bShouldTransact)
+		{
+			GEditor->EndTransaction();
+		}
 	}
 }
 
@@ -237,7 +280,7 @@ void FObjectMixerEditorListRowData::PropagateChangesToSimilarSelectedRowProperti
 	FObjectMixerEditorListRowData* RowData = FObjectMixerUtils::GetRowData(InRow);
 	
 	const TWeakPtr<IPropertyHandle>* HandlePtr = RowData->PropertyNamesToHandles.Find(PropertyPropagationInfo.PropertyName);
-	if (HandlePtr->IsValid())
+	if (HandlePtr && HandlePtr->IsValid())
 	{
 		const TArray<TSharedPtr<ISceneOutlinerTreeItem>> OtherSelectedItems = RowData->GetSelectedTreeViewItems();
 		if (OtherSelectedItems.Num())
