Object Mixer:

-Fix: Changing Parameters on a Multi-Selection Only Change Values Of Lights On Screen
-Fix: Component selection is cleared in mixer panel when panel is refreshed

diff --git a/Engine/Plugins/Editor/ObjectMixer/ObjectMixer/Source/ObjectMixer/Private/Views/List/Modes/ObjectMixerOutlinerMode.cpp b/Engine/Plugins/Editor/ObjectMixer/ObjectMixer/Source/ObjectMixer/Private/Views/List/Modes/ObjectMixerOutlinerMode.cpp
--- a/Engine/Plugins/Editor/ObjectMixer/ObjectMixer/Source/ObjectMixer/Private/Views/List/Modes/ObjectMixerOutlinerMode.cpp
+++ b/Engine/Plugins/Editor/ObjectMixer/ObjectMixer/Source/ObjectMixer/Private/Views/List/Modes/ObjectMixerOutlinerMode.cpp
@@ -105,7 +105,7 @@ namespace ObjectMixerOutliner
 		return false;
 	}
 
-	bool FComponentSelector::operator()(const TWeakPtr<ISceneOutlinerTreeItem>& Item, TObjectPtr<UActorComponent>& DataOut) const
+	bool FComponentSelector::operator()(const TWeakPtr<ISceneOutlinerTreeItem>& Item, UActorComponent*& DataOut) const
 	{
 		if (TSharedPtr<ISceneOutlinerTreeItem> ItemPtr = Item.Pin())
 		{
@@ -564,12 +564,6 @@ void FObjectMixerOutlinerMode::OnLevelSelectionChanged(UObject* Obj)
 		SceneOutliner->FullRefresh();
 		return;
 	}
-
-	if (bShouldPauseSelectionSyncFromEditorTemporarily)
-	{
-		bShouldPauseSelectionSyncFromEditorTemporarily = false;
-		return;
-	}
 	
 	// If the SceneOutliner's reentrant flag is set, the selection change has already been handled in the outliner class
 	if (!SceneOutliner->GetIsReentrant())
@@ -630,6 +624,130 @@ void FObjectMixerOutlinerMode::OnPostLoadMapWithWorld(UWorld* World)
 	SceneOutliner->FullRefresh();
 }
 
+bool FObjectMixerOutlinerMode::HasActorSelectionChanged(TArray<AActor*>& OutSelectedActors, bool& bOutAreAnyInPIE)
+{
+	bool bHasSelectionChanged = false;
+	for (AActor* Actor : OutSelectedActors)
+	{
+		if (!bOutAreAnyInPIE && Actor && Actor->GetPackage()->HasAnyPackageFlags(PKG_PlayInEditor))
+		{
+			bOutAreAnyInPIE = true;
+		}
+		if (!GEditor->GetSelectedActors()->IsSelected(Actor))
+		{
+			bHasSelectionChanged = true;
+			break;
+		}
+	}
+
+	for (FSelectionIterator SelectionIt(*GEditor->GetSelectedActors()); SelectionIt && !bHasSelectionChanged; ++SelectionIt)
+	{
+		const AActor* Actor = CastChecked< AActor >(*SelectionIt);
+		if (!bOutAreAnyInPIE && Actor->GetPackage()->HasAnyPackageFlags(PKG_PlayInEditor))
+		{
+			bOutAreAnyInPIE = true;
+		}
+		if (!OutSelectedActors.Contains(Actor))
+		{
+			// Actor has been deselected
+			bHasSelectionChanged = true;
+
+			// If actor was a group actor, remove its members from the ActorsToSelect list
+			const AGroupActor* DeselectedGroupActor = Cast<AGroupActor>(Actor);
+			if (DeselectedGroupActor)
+			{
+				TArray<AActor*> GroupActors;
+				DeselectedGroupActor->GetGroupActors(GroupActors);
+
+				for (auto* GroupActor : GroupActors)
+				{
+					OutSelectedActors.Remove(GroupActor);
+				}
+
+			}
+		}
+	}
+
+	return bHasSelectionChanged;
+}
+
+bool FObjectMixerOutlinerMode::HasComponentSelectionChanged(TArray<UActorComponent*>& OutSelectedComponents, bool& bOutAreAnyInPIE)
+{
+	bool bHasSelectionChanged = false;
+	for (UActorComponent* Component : OutSelectedComponents)
+	{
+		if (!bOutAreAnyInPIE && Component && Component->GetPackage()->HasAnyPackageFlags(PKG_PlayInEditor))
+		{
+			bOutAreAnyInPIE = true;
+		}
+		if (!GEditor->GetSelectedComponents()->IsSelected(Component))
+		{
+			bHasSelectionChanged = true;
+			break;
+		}
+	}
+
+	for (FSelectionIterator SelectionIt(*GEditor->GetSelectedComponents()); SelectionIt && !bHasSelectionChanged; ++SelectionIt)
+	{
+		const UActorComponent* Actor = CastChecked< UActorComponent >(*SelectionIt);
+		if (!bOutAreAnyInPIE && Actor->GetPackage()->HasAnyPackageFlags(PKG_PlayInEditor))
+		{
+			bOutAreAnyInPIE = true;
+		}
+		if (!OutSelectedComponents.Contains(Actor))
+		{
+			// Actor has been deselected
+			bHasSelectionChanged = true;
+		}
+	}
+
+	return bHasSelectionChanged;
+}
+
+void FObjectMixerOutlinerMode::SelectActorsInEditor(const TArray<AActor*>& InSelectedActors)
+{
+	GEditor->GetSelectedActors()->Modify();
+				
+	// We'll batch selection changes instead by using BeginBatchSelectOperation()
+	GEditor->GetSelectedActors()->BeginBatchSelectOperation();
+
+	// Clear the selection
+	GEditor->GetSelectedActors()->DeselectAll();
+				
+	for (AActor* Actor : InSelectedActors)
+	{
+		constexpr bool bShouldSelect = true;
+		constexpr bool bNotifyAfterSelect = false;
+		constexpr bool bSelectEvenIfHidden = true;
+		GEditor->SelectActor(Actor, bShouldSelect, bNotifyAfterSelect, bSelectEvenIfHidden);
+	}
+
+	// Commit selection changes
+	GEditor->GetSelectedActors()->EndBatchSelectOperation(/*bNotify*/false);
+}
+
+void FObjectMixerOutlinerMode::SelectComponentsInEditor(const TArray<UActorComponent*>& InSelectedComponents)
+{	
+	GEditor->GetSelectedComponents()->Modify();
+				
+	// We'll batch selection changes instead by using BeginBatchSelectOperation()
+	GEditor->GetSelectedComponents()->BeginBatchSelectOperation();
+
+	// Clear the selection
+	GEditor->GetSelectedComponents()->DeselectAll();
+				
+	for (UActorComponent* Component : InSelectedComponents)
+	{
+		constexpr bool bShouldSelect = true;
+		constexpr bool bNotifyAfterSelect = false;
+		constexpr bool bSelectEvenIfHidden = true;
+		GEditor->SelectComponent(Component, bShouldSelect, bNotifyAfterSelect, bSelectEvenIfHidden);
+	}
+
+	// Commit selection changes
+	GEditor->GetSelectedComponents()->EndBatchSelectOperation(/*bNotify*/false);
+}
+
 void FObjectMixerOutlinerMode::OnActorLabelChanged(AActor* ChangedActor)
 {
 	if (!ensure(ChangedActor))
@@ -1675,84 +1793,35 @@ void FObjectMixerOutlinerMode::OnItemSelectionChanged(FSceneOutlinerTreeItemPtr
 {
 	if (GetDefault<UObjectMixerEditorSettings>()->bSyncSelection)
 	{
-		bShouldPauseSelectionSyncFromEditorTemporarily = true;
+		bool bAreAnyInPIE = false;
+
+		// Actors
+		SynchronizeSelectedActorDescs();
 		TArray<AActor*> SelectedActors = Selection.GetData<AActor*>(ObjectMixerOutliner::FActorSelector());
+		bool bHasActorSelectionChanged = HasActorSelectionChanged(SelectedActors, bAreAnyInPIE);
 		
-		SynchronizeSelectedActorDescs();
-
-		bool bChanged = false;
-		bool bAnyInPIE = false;
-		for (auto* Actor : SelectedActors)
+		// Components
+		TArray<UActorComponent*> SelectedComponents = Selection.GetData<UActorComponent*>(ObjectMixerOutliner::FComponentSelector());
+		bool bHasComponentSelectionChanged = HasComponentSelectionChanged(SelectedComponents, bAreAnyInPIE);
+		
+		// If there's a discrepancy, update the selected objects to reflect the list.
+		if (bHasActorSelectionChanged || bHasComponentSelectionChanged)
 		{
-			if (!bAnyInPIE && Actor && Actor->GetPackage()->HasAnyPackageFlags(PKG_PlayInEditor))
-			{
-				bAnyInPIE = true;
-			}
-			if (!GEditor->GetSelectedActors()->IsSelected(Actor))
-			{
-				bChanged = true;
-				break;
-			}
-		}
+			const FScopedTransaction Transaction(LOCTEXT("ClickingOnActorsAndComponents", "Clicking on Actors & Components"), !bAreAnyInPIE);
 
-		for (FSelectionIterator SelectionIt(*GEditor->GetSelectedActors()); SelectionIt && !bChanged; ++SelectionIt)
-		{
-			const AActor* Actor = CastChecked< AActor >(*SelectionIt);
-			if (!bAnyInPIE && Actor->GetPackage()->HasAnyPackageFlags(PKG_PlayInEditor))
-			{
-				bAnyInPIE = true;
-			}
-			if (!SelectedActors.Contains(Actor))
+			if (bHasActorSelectionChanged)
 			{
-				// Actor has been deselected
-				bChanged = true;
-
-				// If actor was a group actor, remove its members from the ActorsToSelect list
-				const AGroupActor* DeselectedGroupActor = Cast<AGroupActor>(Actor);
-				if (DeselectedGroupActor)
-				{
-					TArray<AActor*> GroupActors;
-					DeselectedGroupActor->GetGroupActors(GroupActors);
-
-					for (auto* GroupActor : GroupActors)
-					{
-						SelectedActors.Remove(GroupActor);
-					}
-
-				}
+				SelectActorsInEditor(SelectedActors);
 			}
-		}
-
-		// If there's a discrepancy, update the selected actors to reflect this list.
-		if (bChanged)
-		{
-			const FScopedTransaction Transaction(NSLOCTEXT("UnrealEd", "ClickingOnActors", "Clicking on Actors"), !bAnyInPIE);
-			GEditor->GetSelectedActors()->Modify();
-
-			// We'll batch selection changes instead by using BeginBatchSelectOperation()
-			GEditor->GetSelectedActors()->BeginBatchSelectOperation();
-
-			// Clear the selection.
-			GEditor->SelectNone(false, true, true);
-
-			const bool bShouldSelect = true;
-			const bool bNotifyAfterSelect = false;
-			const bool bSelectEvenIfHidden = true;	// @todo outliner: Is this actually OK?
-			for (auto* Actor : SelectedActors)
+			
+			if (bHasComponentSelectionChanged)
 			{
-				UE_LOG(LogObjectMixerEditor, Verbose, TEXT("Clicking on Actor (Object Mixer): %s (%s)"), *Actor->GetClass()->GetName(), *Actor->GetActorLabel());
-				GEditor->SelectActor(Actor, bShouldSelect, bNotifyAfterSelect, bSelectEvenIfHidden);
+				SelectComponentsInEditor(SelectedComponents);
 			}
 
-			// Commit selection changes
-			GEditor->GetSelectedActors()->EndBatchSelectOperation(/*bNotify*/false);
-
 			SceneOutliner->RefreshSelection();
 		}
 	}
-
-	// Set this outliner as the most recently interacted with
-	SetAsMostRecentOutliner();
 }
 
 void FObjectMixerOutlinerMode::OnItemDoubleClick(FSceneOutlinerTreeItemPtr Item)
@@ -2151,6 +2220,46 @@ FFolder FObjectMixerOutlinerMode::GetWorldDefaultRootFolder() const
 	return FFolder::GetWorldRootFolder(RepresentingWorld.Get());
 }
 
+void FObjectMixerOutlinerMode::SynchronizeComponentSelection()
+{
+	USelection* SelectedComponents = GEditor->GetSelectedComponents();
+
+	// Deselect components in the tree that are no longer selected in the world
+	const FSceneOutlinerItemSelection Selection(SceneOutliner->GetSelection());
+	auto DeselectComponents = [this](FComponentTreeItem& Item)
+	{
+		if (!Item.Component.IsValid() || !Item.Component.Get()->IsSelected())
+		{
+			SceneOutliner->SetItemSelection(Item.AsShared(), false);
+		}
+	};
+	Selection.ForEachItem<FComponentTreeItem>(DeselectComponents);
+
+	// See if the tree view selector is pointing at a selected item
+	bool bSelectorInSelectionSet = false;
+
+	TArray<FSceneOutlinerTreeItemPtr> ComponentItems;
+	for (FSelectionIterator SelectionIt(*SelectedComponents); SelectionIt; ++SelectionIt)
+	{
+		UActorComponent* Component = CastChecked< UActorComponent >(*SelectionIt);
+		if (FSceneOutlinerTreeItemPtr ComponentItem = SceneOutliner->GetTreeItem(Component))
+		{
+			if (!bSelectorInSelectionSet && SceneOutliner->HasSelectorFocus(ComponentItem))
+			{
+				bSelectorInSelectionSet = true;
+			}
+
+			ComponentItems.Add(ComponentItem);
+		}
+	}
+
+	// If NOT bSelectorInSelectionSet then we want to just move the selector to the first selected item.
+	ESelectInfo::Type SelectInfo = bSelectorInSelectionSet ? ESelectInfo::Direct : ESelectInfo::OnMouseClick;
+	SceneOutliner->AddToSelection(ComponentItems, SelectInfo);
+
+	FSceneOutlinerDelegates::Get().SelectionChanged.Broadcast();
+}
+
 void FObjectMixerOutlinerMode::SynchronizeSelectedActorDescs()
 {
 	if (UWorldPartitionSubsystem* WorldPartitionSubsystem = UWorld::GetSubsystem<UWorldPartitionSubsystem>(RepresentingWorld.Get()))
@@ -3008,6 +3117,7 @@ void FObjectMixerOutlinerMode::SynchronizeSelection()
 	if (GetDefault<UObjectMixerEditorSettings>()->bSyncSelection)
 	{
 		SynchronizeActorSelection();
+		SynchronizeComponentSelection();
 		SynchronizeSelectedActorDescs();
 	}
 }
diff --git a/Engine/Plugins/Editor/ObjectMixer/ObjectMixer/Source/ObjectMixer/Private/Views/List/Modes/OutlinerColumns/ObjectMixerOutlinerColumns.cpp b/Engine/Plugins/Editor/ObjectMixer/ObjectMixer/Source/ObjectMixer/Private/Views/List/Modes/OutlinerColumns/ObjectMixerOutlinerColumns.cpp
--- a/Engine/Plugins/Editor/ObjectMixer/ObjectMixer/Source/ObjectMixer/Private/Views/List/Modes/OutlinerColumns/ObjectMixerOutlinerColumns.cpp
+++ b/Engine/Plugins/Editor/ObjectMixer/ObjectMixer/Source/ObjectMixer/Private/Views/List/Modes/OutlinerColumns/ObjectMixerOutlinerColumns.cpp
@@ -105,7 +105,6 @@ void FObjectMixerOutlinerSoloColumn::SortItems(TArray<FSceneOutlinerTreeItemPtr>
 FSlateColor FObjectMixerOutlinerSoloColumn::GetForegroundColor(
 	FSceneOutlinerTreeItemRef TreeItem, const STableRow<FSceneOutlinerTreeItemPtr>* Row) const
 {
-	// todo when hybrid rows are back: const TSharedRef<ISceneOutlinerTreeItem> RowPtr = GetHybridChildOrRowItemIfNull(TreeItem);
 	if (FObjectMixerEditorListRowData* RowData = FObjectMixerUtils::GetRowData(TreeItem))
 	{
 		if (RowData->GetIsSelected(TreeItem) || RowData->GetRowSoloState())
@@ -126,7 +125,7 @@ void FObjectMixerOutlinerSoloColumn::OnClickSoloIcon(const FSceneOutlinerTreeIte
 {
 	if (FObjectMixerEditorListRowData* RowData = FObjectMixerUtils::GetRowData(RowPtr))
 	{
-		if (SObjectMixerEditorList* ListView = RowData->GetListView())
+		if (TSharedPtr<SObjectMixerEditorList> ListView = RowData->GetListView().Pin())
 		{
 			const bool bNewSolo = !RowData->GetRowSoloState();
 								
@@ -225,26 +224,31 @@ void FObjectMixerOutlinerPropertyColumn::OnPropertyChanged(const FPropertyChange
 	{
 		if (FObjectMixerEditorListRowData* RowData = FObjectMixerUtils::GetRowData(TreeItem))
 		{
-			const EPropertyValueSetFlags::Type Flag =
-				 Event.ChangeType == EPropertyChangeType::Interactive ?
-					 EPropertyValueSetFlags::InteractiveChange : EPropertyValueSetFlags::DefaultFlags;
-
-			FObjectMixerEditorListRowData::FPropertyPropagationInfo PropagationInfo; 
-			PropagationInfo.RowIdentifier = TreeItem->GetID();
-			PropagationInfo.PropertyName = Property->GetFName();
-			PropagationInfo.PropertyValueSetFlags = Flag;
-		
-			if (Flag == EPropertyValueSetFlags::InteractiveChange)
+			if (const TSharedPtr<SObjectMixerEditorList> ListView = RowData->GetListView().Pin())
 			{
-				RowData->PropagateChangesToSimilarSelectedRowProperties(TreeItem.ToSharedRef(), PropagationInfo);
-			}
-			else
-			{
-				// If not an interactive change, schedule property propagation on next frame
-				if (SObjectMixerEditorList* ListView = RowData->GetListView())
+				// If fewer than 2 items are selected, there is nothing to propagate. Early out.
+				if (ListView->GetSelectedTreeViewItemCount() < 2)
+				{
+					return;
+				}
+				
+				const EPropertyValueSetFlags::Type Flag =
+					 Event.ChangeType & EPropertyChangeType::Interactive ?
+						 EPropertyValueSetFlags::InteractiveChange : EPropertyValueSetFlags::DefaultFlags;
+
+				FObjectMixerEditorListRowData::FPropertyPropagationInfo PropagationInfo; 
+				PropagationInfo.RowIdentifier = TreeItem->GetID();
+				PropagationInfo.PropertyName = Property->GetFName();
+				PropagationInfo.PropertyValueSetFlags = Flag;
+		
+				if (Flag == EPropertyValueSetFlags::InteractiveChange)
+				{
+					RowData->PropagateChangesToSimilarSelectedRowProperties(TreeItem.ToSharedRef(), PropagationInfo);
+				}
+				else
 				{
+					// If not an interactive change, schedule property propagation on next frame
 					ListView->AddToPendingPropertyPropagations(PropagationInfo);
-					ListView->RequestRebuildList();
 				}
 			}
 		}
diff --git a/Engine/Plugins/Editor/ObjectMixer/ObjectMixer/Source/ObjectMixer/Private/Views/List/ObjectMixerEditorListRowData.cpp b/Engine/Plugins/Editor/ObjectMixer/ObjectMixer/Source/ObjectMixer/Private/Views/List/ObjectMixerEditorListRowData.cpp
--- a/Engine/Plugins/Editor/ObjectMixer/ObjectMixer/Source/ObjectMixer/Private/Views/List/ObjectMixerEditorListRowData.cpp
+++ b/Engine/Plugins/Editor/ObjectMixer/ObjectMixer/Source/ObjectMixer/Private/Views/List/ObjectMixerEditorListRowData.cpp
@@ -12,8 +12,8 @@
 
 const TArray<TObjectPtr<UObjectMixerObjectFilter>>& FObjectMixerEditorListRowData::GetObjectFilterInstances() const
 {
-	SObjectMixerEditorList* ListView = GetListView();
-	check (ListView)
+	TSharedPtr<SObjectMixerEditorList> ListView = GetListView().Pin();
+	check (ListView.IsValid());
 
 	const TSharedPtr<FObjectMixerEditorList> PinnedListModel = ListView->GetListModelPtr().Pin();
 	check (PinnedListModel);
@@ -23,7 +23,7 @@ const TArray<TObjectPtr<UObjectMixerObjectFilter>>& FObjectMixerEditorListRowDat
 
 const UObjectMixerObjectFilter* FObjectMixerEditorListRowData::GetMainObjectFilterInstance() const
 {
-	if (SObjectMixerEditorList* ListView = GetListView())
+	if (TSharedPtr<SObjectMixerEditorList> ListView = GetListView().Pin(); ListView.IsValid())
 	{
 		if (const TSharedPtr<FObjectMixerEditorList> PinnedListModel = ListView->GetListModelPtr().Pin())
 		{
@@ -36,17 +36,25 @@ const UObjectMixerObjectFilter* FObjectMixerEditorListRowData::GetMainObjectFilt
 
 bool FObjectMixerEditorListRowData::GetIsTreeViewItemExpanded(const TSharedRef<ISceneOutlinerTreeItem> InRow)
 {
-	return GetListView()->IsTreeViewItemExpanded(InRow);
+	if (TSharedPtr<SObjectMixerEditorList> ListView = GetListView().Pin(); ListView.IsValid())
+	{
+		return ListView->IsTreeViewItemExpanded(InRow);
+	}
+
+	return false;
 }
 
 void FObjectMixerEditorListRowData::SetIsTreeViewItemExpanded(const TSharedRef<ISceneOutlinerTreeItem> InRow, const bool bNewExpanded)
 {
-	GetListView()->SetTreeViewItemExpanded(InRow, bNewExpanded);
+	if (TSharedPtr<SObjectMixerEditorList> ListView = GetListView().Pin(); ListView.IsValid())
+	{
+		ListView->SetTreeViewItemExpanded(InRow, bNewExpanded);
+	}
 }
 
 bool FObjectMixerEditorListRowData::GetIsSelected(const TSharedRef<ISceneOutlinerTreeItem> InRow)
 {
-	if (SObjectMixerEditorList* ListView = GetListView())
+	if (TSharedPtr<SObjectMixerEditorList> ListView = GetListView().Pin(); ListView.IsValid())
 	{
 		return ListView->IsTreeViewItemSelected(InRow);
 	}
@@ -56,7 +64,7 @@ bool FObjectMixerEditorListRowData::GetIsSelected(const TSharedRef<ISceneOutline
 
 void FObjectMixerEditorListRowData::SetIsSelected(const TSharedRef<ISceneOutlinerTreeItem> InRow, const bool bNewSelected)
 {
-	if (SObjectMixerEditorList* ListView = GetListView())
+	if (TSharedPtr<SObjectMixerEditorList> ListView = GetListView().Pin(); ListView.IsValid())
 	{
 		return ListView->SetTreeViewItemSelected(InRow, bNewSelected);
 	}
@@ -102,14 +110,19 @@ FText FObjectMixerEditorListRowData::GetDisplayName(TSharedPtr<ISceneOutlinerTre
 	return FText::GetEmpty();
 }
 
-SObjectMixerEditorList* FObjectMixerEditorListRowData::GetListView() const
+TWeakPtr<SObjectMixerEditorList> FObjectMixerEditorListRowData::GetListView() const
 {
-	return StaticCast<SObjectMixerEditorList*>(SceneOutlinerPtr);
+	if (SceneOutlinerPtr.IsValid())
+	{
+		return StaticCastSharedPtr<SObjectMixerEditorList>(SceneOutlinerPtr.Pin());
+	}
+
+	return nullptr;
 }
 
 TArray<TSharedPtr<ISceneOutlinerTreeItem>> FObjectMixerEditorListRowData::GetSelectedTreeViewItems() const
 {
-	if (SObjectMixerEditorList* ListView = GetListView())
+	if (TSharedPtr<SObjectMixerEditorList> ListView = GetListView().Pin(); ListView.IsValid())
 	{
 		return ListView->GetSelectedTreeViewItems();
 	}
@@ -121,7 +134,7 @@ void FObjectMixerEditorListRowData::OnChangeVisibility(const FSceneOutlinerTreeI
 {
 	if (FObjectMixerEditorListRowData* RowData = FObjectMixerUtils::GetRowData(TreeItem))
 	{
-		if (SObjectMixerEditorList* ListView = RowData->GetListView())
+		if (TSharedPtr<SObjectMixerEditorList> ListView = GetListView().Pin(); ListView.IsValid())
 		{
 			ListView->ClearSoloRows();
 		}
@@ -152,12 +165,12 @@ void FObjectMixerEditorListRowData::SetRowSoloState(const bool bNewSolo)
 
 void FObjectMixerEditorListRowData::ClearSoloRows() const
 {
-	if (SObjectMixerEditorList* ListView = GetListView())
+	if (TSharedPtr<SObjectMixerEditorList> ListView = GetListView().Pin(); ListView.IsValid())
 	{
 		ListView->ClearSoloRows();
 	}
 }
-
+DECLARE_STATS_GROUP(TEXT("ObjectMixer"), STATGROUP_ObjectMixer, STATCAT_Advanced);
 void SetValueOnSelectedItems(
 	const FString& ValueAsString, const TArray<TSharedPtr<ISceneOutlinerTreeItem>>& OtherSelectedItems,
 	const FName& PropertyName, const TSharedPtr<ISceneOutlinerTreeItem> PinnedItem,
@@ -165,13 +178,13 @@ void SetValueOnSelectedItems(
 {
 	if (!ValueAsString.IsEmpty())
 	{
-		const EPropertyValueSetFlags::Type Flags = InFlags | EPropertyValueSetFlags::NotTransactable;
-		const bool bShouldTransact = InFlags & EPropertyValueSetFlags::DefaultFlags;
+		// Skip transactions on interactive and explicitly non-transactable value set
+		const bool bShouldTransact = !(InFlags & EPropertyValueSetFlags::NotTransactable) && !(InFlags & EPropertyValueSetFlags::InteractiveChange);
 
-		if (bShouldTransact && GEditor->CanTransact())
+		if (bShouldTransact && GEditor->CanTransact() && !GEditor->IsTransactionActive())
 		{
 			GEditor->BeginTransaction(
-			   NSLOCTEXT("ObjectMixerEditor","OnPropertyChangedTransaction", "Object Mixer - Edit Selected Row Properties"));
+			   NSLOCTEXT("ObjectMixerEditor","BulkOnPropertyChangedTransaction", "Object Mixer - Bulk Edit Selected Row Properties"));
 		}
 		
 		for (const TSharedPtr<ISceneOutlinerTreeItem>& SelectedRow : OtherSelectedItems)
@@ -181,12 +194,6 @@ void SetValueOnSelectedItems(
 				continue;
 			}
 
-			FObjectMixerEditorListRowData* RowData = FObjectMixerUtils::GetRowData(SelectedRow);
-			if (!RowData)
-			{
-				continue;
-			}
-
 			// Skip folders
 			if (FObjectMixerUtils::AsFolderRow(SelectedRow))
 			{
@@ -195,25 +202,37 @@ void SetValueOnSelectedItems(
 
 			UObject* ObjectToModify = FObjectMixerUtils::GetRowObject(SelectedRow, true);
 			
-			if (IsValid(ObjectToModify))
-			{
-				ObjectToModify->Modify();
-			}
-			else
+			if (!IsValid(ObjectToModify))
 			{
-				UE_LOG(LogObjectMixerEditor, Warning, TEXT("%hs: Row '%s' has no valid associated object to modify."), __FUNCTION__, *RowData->GetDisplayName(SelectedRow).ToString());
+				UE_LOG(LogObjectMixerEditor, Warning, TEXT("%hs: Row '%s' has no valid associated object to modify."), __FUNCTION__, *SelectedRow->GetDisplayString());
 				continue;
 			}
-		
+
 			// Use handles if valid, otherwise use ImportText. Need to use the handles to ensure the Blueprints update properly.
-			
-			// Transactions are handled automatically by the handles, so no need to start a new transaction.
-			if (const TWeakPtr<IPropertyHandle>* SelectedHandlePtr = RowData->PropertyNamesToHandles.Find(PropertyName))
+			if (FObjectMixerEditorListRowData* RowData = FObjectMixerUtils::GetRowData(SelectedRow))
 			{
-				if (SelectedHandlePtr->IsValid())
-				{				
-					SelectedHandlePtr->Pin()->SetValueFromFormattedString(ValueAsString, Flags);
-					continue;
+				// Transactions are handled automatically by the handles, so no need to start a new transaction.
+				if (const TWeakPtr<IPropertyHandle>* SelectedHandlePtr = RowData->PropertyNamesToHandles.Find(PropertyName))
+				{
+					if (SelectedHandlePtr->IsValid())
+					{
+						bool bShouldProceed = true;
+
+						// If we need to transact, there's a possibility objects have been reconstructed
+						// so we need to ensure we have handles matching ObjectToModify
+						if (bShouldTransact)
+						{
+							TArray<UObject*> OuterObjects;
+							SelectedHandlePtr->Pin()->GetOuterObjects(OuterObjects);
+							bShouldProceed = OuterObjects.Contains(ObjectToModify);
+						}
+						
+						if (bShouldProceed)
+						{
+							SelectedHandlePtr->Pin()->SetValueFromFormattedString(ValueAsString, InFlags);
+							continue;
+						}
+					}
 				}
 			}
 
@@ -233,26 +252,23 @@ void SetValueOnSelectedItems(
 							? EPropertyChangeType::Interactive
 							: EPropertyChangeType::ValueSet;
 
-					// Set the actual property value
+					// Set the actual property value and propagate to outer chain
 					PropertyToChange->ImportText_Direct(*ValueAsString, ValuePtr, ObjectToModify, PPF_None);
-					FPropertyChangedEvent ChangeEvent(
-						PropertyToChange,
-						ChangeType,
-						MakeArrayView({ObjectToModify}));
-					ObjectToModify->PostEditChangeProperty(ChangeEvent);
 
-					// Propagate to outers
+					TArray<UObject*> ObjectsToModify = {ObjectToModify};
 					UObject* Outer = ObjectToModify->GetOuter();
 					while (Outer) 
 					{
-						FPropertyChangedEvent ActorChangeEvent(
-							PropertyToChange,
-							ChangeType,
-							MakeArrayView({Outer}));
-						Outer->PostEditChangeProperty(ActorChangeEvent);
+						ObjectsToModify.Add(Outer);
 
 						Outer = Outer->GetOuter();
 					}
+					
+					FPropertyChangedEvent ChangeEvent(
+						PropertyToChange,
+						ChangeType,
+						MakeArrayView(ObjectsToModify));
+					ObjectToModify->PostEditChangeProperty(ChangeEvent);
 				}
 			}
 		}
@@ -264,17 +280,17 @@ void SetValueOnSelectedItems(
 	}
 }
 
-void FObjectMixerEditorListRowData::PropagateChangesToSimilarSelectedRowProperties(
+bool FObjectMixerEditorListRowData::PropagateChangesToSimilarSelectedRowProperties(
 	const TSharedRef<ISceneOutlinerTreeItem> InRow, const FPropertyPropagationInfo PropertyPropagationInfo)
 {
 	if (PropertyPropagationInfo.PropertyName == NAME_None)
 	{
-		return;
+		return true;
 	}
 
 	if (!GetIsSelected(InRow))
 	{
-		return;
+		return true;
 	}
 
 	FObjectMixerEditorListRowData* RowData = FObjectMixerUtils::GetRowData(InRow);
@@ -292,7 +308,10 @@ void FObjectMixerEditorListRowData::PropagateChangesToSimilarSelectedRowProperti
 				ValueAsString, OtherSelectedItems, PropertyPropagationInfo.PropertyName,
 				InRow, PropertyPropagationInfo.PropertyValueSetFlags);
 		}
+		return true;
 	}
+
+	return false;
 }
 
 const FObjectMixerEditorListRowData::FTransientEditorVisibilityRules& FObjectMixerEditorListRowData::GetVisibilityRules() const
diff --git a/Engine/Plugins/Editor/ObjectMixer/ObjectMixer/Source/ObjectMixer/Private/Views/List/ObjectMixerUtils.cpp b/Engine/Plugins/Editor/ObjectMixer/ObjectMixer/Source/ObjectMixer/Private/Views/List/ObjectMixerUtils.cpp
--- a/Engine/Plugins/Editor/ObjectMixer/ObjectMixer/Source/ObjectMixer/Private/Views/List/ObjectMixerUtils.cpp
+++ b/Engine/Plugins/Editor/ObjectMixer/ObjectMixer/Source/ObjectMixer/Private/Views/List/ObjectMixerUtils.cpp
@@ -40,10 +40,7 @@ FObjectMixerEditorListRowActor* FObjectMixerUtils::AsActorRow(TSharedPtr<ISceneO
 	
 	if (FActorTreeItem* AsActor = InTreeItem->CastTo<FActorTreeItem>())
 	{
-		if (AsActor->Actor.IsValid())
-		{
-			return StaticCast<FObjectMixerEditorListRowActor*>(AsActor);
-		}
+		return StaticCast<FObjectMixerEditorListRowActor*>(AsActor);
 	}
 
 	return nullptr;
@@ -60,10 +57,7 @@ FObjectMixerEditorListRowComponent* FObjectMixerUtils::AsComponentRow(TSharedPtr
 	
 	if (FComponentTreeItem* AsComponent = InTreeItem->CastTo<FComponentTreeItem>())
 	{
-		if (AsComponent->Component.IsValid())
-		{
-			return StaticCast<FObjectMixerEditorListRowComponent*>(AsComponent);
-		}
+		return StaticCast<FObjectMixerEditorListRowComponent*>(AsComponent);
 	}
 
 	return nullptr;
@@ -73,10 +67,7 @@ FObjectMixerEditorListRowUObject* FObjectMixerUtils::AsObjectRow(TSharedPtr<ISce
 {
 	if (FObjectMixerEditorListRowUObject* AsObject = InTreeItem->CastTo<FObjectMixerEditorListRowUObject>())
 	{
-		if (IsValid(AsObject->ObjectPtr))
-		{
-			return AsObject;
-		}
+		return AsObject;
 	}
 
 	return nullptr;
@@ -86,22 +77,34 @@ FObjectMixerEditorListRowData* FObjectMixerUtils::GetRowData(TSharedPtr<ISceneOu
 {
 	if (FObjectMixerEditorListRowFolder* AsFolder = AsFolderRow(InTreeItem))
 	{
-		return &AsFolder->RowData;
+		if (AsFolder->RowData.IsValid())
+		{
+			return &AsFolder->RowData;
+		}
 	}
 
 	if (FObjectMixerEditorListRowActor* AsActor = AsActorRow(InTreeItem))
 	{
-		return &AsActor->RowData;
+		if (AsActor->RowData.IsValid())
+		{
+			return &AsActor->RowData;
+		}
 	}
 
 	if (FObjectMixerEditorListRowComponent* AsComponent = AsComponentRow(InTreeItem))
 	{
-		return &AsComponent->RowData;
+		if (AsComponent->RowData.IsValid())
+		{
+			return &AsComponent->RowData;
+		}
 	}
 
 	if (FObjectMixerEditorListRowUObject* AsObject = AsObjectRow(InTreeItem))
 	{
-		return &AsObject->RowData;
+		if (AsObject->RowData.IsValid())
+		{
+			return &AsObject->RowData;
+		}
 	}
 
 	return nullptr;
@@ -115,18 +118,28 @@ UObject* FObjectMixerUtils::GetRowObject(TSharedPtr<ISceneOutlinerTreeItem> InTr
 		{
 			return ActorRow->RowData.GetHybridComponent();
 		}
-		
-		return ActorRow->Actor.Get();
+
+		if (ActorRow->Actor.IsValid() && !ActorRow->Actor.IsStale())
+		{
+			return ActorRow->Actor.Get();
+		}
+
+		return ActorRow->OriginalObjectSoftPtr.Get();
 	}
 
 	if (const FObjectMixerEditorListRowComponent* ComponentRow = AsComponentRow(InTreeItem))
 	{
-		return ComponentRow->Component.Get();
+		if (ComponentRow->Component.IsValid() && !ComponentRow->Component.IsStale())
+		{
+			return ComponentRow->Component.Get();
+		}
+
+		return ComponentRow->OriginalObjectSoftPtr.Get();
 	}
 
 	if (const FObjectMixerEditorListRowUObject* ObjectRow = AsObjectRow(InTreeItem))
 	{
-		return ObjectRow->ObjectPtr.Get();
+		ObjectRow->ObjectSoftPtr.Get();
 	}
 
 	return nullptr;
@@ -168,7 +181,7 @@ bool FObjectMixerUtils::IsObjectRefInCollection(const FName& CollectionName, TSh
 {	
 	if (const UObject* Object = GetRowObject(InTreeItem))
 	{
-		if (SObjectMixerEditorList* ListView = GetRowData(InTreeItem)->GetListView())
+		if (TSharedPtr<SObjectMixerEditorList> ListView = GetRowData(InTreeItem)->GetListView().Pin())
 		{
 			if (const TSharedPtr<FObjectMixerEditorList> ListModel = ListView->GetListModelPtr().Pin())
 			{
diff --git a/Engine/Plugins/Editor/ObjectMixer/ObjectMixer/Source/ObjectMixer/Private/Views/List/RowTypes/ObjectMixerEditorListRowActor.cpp b/Engine/Plugins/Editor/ObjectMixer/ObjectMixer/Source/ObjectMixer/Private/Views/List/RowTypes/ObjectMixerEditorListRowActor.cpp
--- a/Engine/Plugins/Editor/ObjectMixer/ObjectMixer/Source/ObjectMixer/Private/Views/List/RowTypes/ObjectMixerEditorListRowActor.cpp
+++ b/Engine/Plugins/Editor/ObjectMixer/ObjectMixer/Source/ObjectMixer/Private/Views/List/RowTypes/ObjectMixerEditorListRowActor.cpp
@@ -12,6 +12,9 @@ const FSceneOutlinerTreeItemType FObjectMixerEditorListRowActor::Type(&FActorTre
 void FObjectMixerEditorListRowActor::OnVisibilityChanged(const bool bNewVisibility)
 {
 	RowData.OnChangeVisibility(SharedThis(this), bNewVisibility);
-								
-	RowData.GetListView()->EvaluateAndSetEditorVisibilityPerRow();
+
+	if (TSharedPtr<SObjectMixerEditorList> ListView = RowData.GetListView().Pin())
+	{
+		ListView->EvaluateAndSetEditorVisibilityPerRow();
+	}
 }
diff --git a/Engine/Plugins/Editor/ObjectMixer/ObjectMixer/Source/ObjectMixer/Private/Views/List/RowTypes/ObjectMixerEditorListRowFolder.cpp b/Engine/Plugins/Editor/ObjectMixer/ObjectMixer/Source/ObjectMixer/Private/Views/List/RowTypes/ObjectMixerEditorListRowFolder.cpp
--- a/Engine/Plugins/Editor/ObjectMixer/ObjectMixer/Source/ObjectMixer/Private/Views/List/RowTypes/ObjectMixerEditorListRowFolder.cpp
+++ b/Engine/Plugins/Editor/ObjectMixer/ObjectMixer/Source/ObjectMixer/Private/Views/List/RowTypes/ObjectMixerEditorListRowFolder.cpp
@@ -9,6 +9,9 @@ const FSceneOutlinerTreeItemType FObjectMixerEditorListRowFolder::Type(&FFolderT
 void FObjectMixerEditorListRowFolder::OnVisibilityChanged(const bool bNewVisibility)
 {
 	RowData.OnChangeVisibility(SharedThis(this), bNewVisibility);
-								
-	RowData.GetListView()->EvaluateAndSetEditorVisibilityPerRow();
+
+	if (TSharedPtr<SObjectMixerEditorList> ListView = RowData.GetListView().Pin())
+	{
+		ListView->EvaluateAndSetEditorVisibilityPerRow();
+	}
 }
diff --git a/Engine/Plugins/Editor/ObjectMixer/ObjectMixer/Source/ObjectMixer/Private/Views/List/RowTypes/ObjectMixerEditorListRowUObject.cpp b/Engine/Plugins/Editor/ObjectMixer/ObjectMixer/Source/ObjectMixer/Private/Views/List/RowTypes/ObjectMixerEditorListRowUObject.cpp
--- a/Engine/Plugins/Editor/ObjectMixer/ObjectMixer/Source/ObjectMixer/Private/Views/List/RowTypes/ObjectMixerEditorListRowUObject.cpp
+++ b/Engine/Plugins/Editor/ObjectMixer/ObjectMixer/Source/ObjectMixer/Private/Views/List/RowTypes/ObjectMixerEditorListRowUObject.cpp
@@ -6,5 +6,5 @@ const FSceneOutlinerTreeItemType FObjectMixerEditorListRowUObject::Type(&ISceneO
 
 FString FObjectMixerEditorListRowUObject::GetDisplayString() const
 {
-	return ObjectPtr ? ObjectPtr.GetName() : FString();
+	return IsValid() ? ObjectSoftPtr.Get()->GetName() : FString();
 }
diff --git a/Engine/Plugins/Editor/ObjectMixer/ObjectMixer/Source/ObjectMixer/Private/Views/List/SObjectMixerEditorList.cpp b/Engine/Plugins/Editor/ObjectMixer/ObjectMixer/Source/ObjectMixer/Private/Views/List/SObjectMixerEditorList.cpp
--- a/Engine/Plugins/Editor/ObjectMixer/ObjectMixer/Source/ObjectMixer/Private/Views/List/SObjectMixerEditorList.cpp
+++ b/Engine/Plugins/Editor/ObjectMixer/ObjectMixer/Source/ObjectMixer/Private/Views/List/SObjectMixerEditorList.cpp
@@ -368,28 +368,37 @@ void SObjectMixerEditorList::AddToPendingPropertyPropagations(
 
 TArray<TSharedPtr<ISceneOutlinerTreeItem>> SObjectMixerEditorList::GetSelectedTreeViewItems() const
 {
-	auto Selection = GetTree().GetSelectedItems();
-	int32 Count = Selection.Num();
-	UE_LOG(LogTemp, Display, TEXT("%i"), Count);
-	return Selection;
+	if (const TSharedPtr<SSceneOutlinerTreeView> Tree = GetTreeView())
+	{
+		return Tree->GetSelectedItems();
+	}
+
+	return {};
 }
 
 int32 SObjectMixerEditorList::GetSelectedTreeViewItemCount() const
-{	
-	return GetTree().GetSelectedItems().Num();
+{
+	if (const TSharedPtr<SSceneOutlinerTreeView> Tree = GetTreeView())
+	{
+		return Tree->GetSelectedItems().Num();
+	}
+
+	return INDEX_NONE;
 }
 
 void SObjectMixerEditorList::SetTreeViewItemSelected(TSharedRef<ISceneOutlinerTreeItem> Item, const bool bNewSelected)
 {
-	GetTreeView()->SetItemSelection(Item, bNewSelected);
+	if (const TSharedPtr<SSceneOutlinerTreeView> Tree = GetTreeView())
+	{
+		Tree->SetItemSelection(Item, bNewSelected);
+	}
 }
 
 bool SObjectMixerEditorList::IsTreeViewItemSelected(TSharedPtr<ISceneOutlinerTreeItem> Item)
 {
 	if (const TSharedPtr<SSceneOutlinerTreeView> Tree = GetTreeView())
 	{
-		TArray<TSharedPtr<ISceneOutlinerTreeItem, ESPMode::ThreadSafe>> SelectedItems = Tree->GetSelectedItems();
-		return SelectedItems.Contains(Item);
+		return Tree->IsItemSelected(Item);
 	}
 	
 	return false;
@@ -626,23 +635,22 @@ void SObjectMixerEditorList::CustomAddToToolbar(TSharedPtr<SHorizontalBox> Toolb
 
 void SObjectMixerEditorList::Tick(const FGeometry& AllottedGeometry, const double InCurrentTime, const float InDeltaTime)
 {
-	SCompoundWidget::Tick(AllottedGeometry, InCurrentTime, InDeltaTime);
-
+	// Early out so we have widgets to act upon next frame
+	SSceneOutliner::Tick(AllottedGeometry, InCurrentTime, InDeltaTime);
+	
 	if (bIsRebuildRequested)
 	{
 		RebuildList();
-		
-		// Early out so we have widgets to act upon next frame
-		SSceneOutliner::Tick(AllottedGeometry, InCurrentTime, InDeltaTime);
-		return;
 	}
 
 	if (PendingPropertyPropagations.Num() > 0)
 	{
 		PropagatePropertyChangesToSelectedRows();
+		if (PendingPropertyPropagations.Num() == 0)
+		{
+			RequestRebuildList();
+		}
 	}
-
-	SSceneOutliner::Tick(AllottedGeometry, InCurrentTime, InDeltaTime);
 }
 
 TSharedRef<SWidget> SObjectMixerEditorList::GenerateHeaderRowContextMenu() const
@@ -1316,12 +1324,12 @@ void SObjectMixerEditorList::PropagatePropertyChangesToSelectedRows()
 {
 	using LambdaType = void(*)(const TSet<TWeakPtr<ISceneOutlinerTreeItem>>&,
 			TSharedPtr<STreeView<TSharedPtr<ISceneOutlinerTreeItem>>>,
-			const TArray<FObjectMixerEditorListRowData::FPropertyPropagationInfo>&);
+			TSet<FObjectMixerEditorListRowData::FPropertyPropagationInfo>&);
 	
 	static LambdaType RecursivelyPropagatePropertyChangesToSelectedRows = [](
 		const TSet<TWeakPtr<ISceneOutlinerTreeItem>>& InObjects,
 		TSharedPtr<STreeView<TSharedPtr<ISceneOutlinerTreeItem>>> InTreeViewPtr,
-		const TArray<FObjectMixerEditorListRowData::FPropertyPropagationInfo>& InPendingPropertyPropagations)
+		TSet<FObjectMixerEditorListRowData::FPropertyPropagationInfo>& InPendingPropagations)
 	{
 		for (const TWeakPtr<ISceneOutlinerTreeItem>& TreeViewItem : InObjects)
 		{
@@ -1331,21 +1339,25 @@ void SObjectMixerEditorList::PropagatePropertyChangesToSelectedRows()
 				{
 					const FSceneOutlinerTreeItemID TreeItemUniqueIdentifier = PinnedItem->GetID();
 
-					if (const FObjectMixerEditorListRowData::FPropertyPropagationInfo* PendingPropagation =
+					if (const FObjectMixerEditorListRowData::FPropertyPropagationInfo* Match =
 						Algo::FindByPredicate(
-							InPendingPropertyPropagations,
+							InPendingPropagations,
 							[&TreeItemUniqueIdentifier](const FObjectMixerEditorListRowData::FPropertyPropagationInfo& Other)
 							{
 								return Other.RowIdentifier == TreeItemUniqueIdentifier;
 							}))
 					{
-						RowData->PropagateChangesToSimilarSelectedRowProperties(PinnedItem.ToSharedRef(), *PendingPropagation);
+						
+						if (RowData->PropagateChangesToSimilarSelectedRowProperties(PinnedItem.ToSharedRef(), *Match))
+						{
+							// Only remove this propagation from the list if it was successfully propagated (or there is no work to do)
+							InPendingPropagations.Remove(*Match);
+						}
 					}
 				}
 
 				RecursivelyPropagatePropertyChangesToSelectedRows(
-					PinnedItem->GetChildren(), InTreeViewPtr, InPendingPropertyPropagations
-				);
+					PinnedItem->GetChildren(), InTreeViewPtr, InPendingPropagations);
 			}
 		}
 	};
@@ -1353,11 +1365,8 @@ void SObjectMixerEditorList::PropagatePropertyChangesToSelectedRows()
 	if (GetSelectedTreeViewItemCount() > 1)
 	{
 		RecursivelyPropagatePropertyChangesToSelectedRows(
-		   GetWeakTreeRootItems(), GetTreeView(), PendingPropertyPropagations
-		);
+			GetWeakTreeRootItems(), GetTreeView(), PendingPropertyPropagations);
 	}
-
-	PendingPropertyPropagations.Empty();
 }
 
 #undef LOCTEXT_NAMESPACE
diff --git a/Engine/Plugins/Editor/ObjectMixer/ObjectMixer/Source/ObjectMixer/Public/Views/List/Modes/ObjectMixerOutlinerMode.h b/Engine/Plugins/Editor/ObjectMixer/ObjectMixer/Source/ObjectMixer/Public/Views/List/Modes/ObjectMixerOutlinerMode.h
--- a/Engine/Plugins/Editor/ObjectMixer/ObjectMixer/Source/ObjectMixer/Public/Views/List/Modes/ObjectMixerOutlinerMode.h
+++ b/Engine/Plugins/Editor/ObjectMixer/ObjectMixer/Source/ObjectMixer/Public/Views/List/Modes/ObjectMixerOutlinerMode.h
@@ -23,7 +23,7 @@ namespace ObjectMixerOutliner
 
 	struct FComponentSelector
 	{
-		bool operator()(const TWeakPtr<ISceneOutlinerTreeItem>& Item, TObjectPtr<UActorComponent>& DataOut) const;
+		bool operator()(const TWeakPtr<ISceneOutlinerTreeItem>& Item, UActorComponent*& DataOut) const;
 	};
 
 	/** Functor which can be used to get weak actor pointers from a selection */
@@ -278,6 +278,11 @@ protected:
 	void OnObjectsReplaced(const TMap<UObject*, UObject*>& ReplacementMap);
 	void OnLevelActorRequestsRename(const AActor* Actor);
 	void OnPostLoadMapWithWorld(UWorld* World);
+
+	bool HasActorSelectionChanged(TArray<AActor*>& OutSelectedActors, bool& bOutAreAnyInPIE);
+	bool HasComponentSelectionChanged(TArray<UActorComponent*>& OutSelectedComponents, bool& bOutAreAnyInPIE);
+	static void SelectActorsInEditor(const TArray<AActor*>& InSelectedActors);
+	static void SelectComponentsInEditor(const TArray<UActorComponent*>& InSelectedComponents);
 	
 	/** Build and up the context menu */
 	TSharedPtr<SWidget> BuildContextMenu();
@@ -288,6 +293,7 @@ protected:
 	bool GetFolderNamesFromPayload(const FSceneOutlinerDragDropPayload& InPayload, TArray<FName>& OutFolders, FFolder::FRootObject& OutCommonRootObject) const;
 	FFolder GetWorldDefaultRootFolder() const;
 
+	void SynchronizeComponentSelection();
 	void SynchronizeSelectedActorDescs();
 
 	void OnActorEditorContextSubsystemChanged();
@@ -341,6 +347,4 @@ protected:
 	
 	TArray<FFilterClassSelectionInfo> FilterClassSelectionInfos;
 	TSharedRef<SWidget> OnGenerateFilterClassMenu();
-
-	bool bShouldPauseSelectionSyncFromEditorTemporarily = false;
 };
diff --git a/Engine/Plugins/Editor/ObjectMixer/ObjectMixer/Source/ObjectMixer/Public/Views/List/ObjectMixerEditorListRowData.h b/Engine/Plugins/Editor/ObjectMixer/ObjectMixer/Source/ObjectMixer/Public/Views/List/ObjectMixerEditorListRowData.h
--- a/Engine/Plugins/Editor/ObjectMixer/ObjectMixer/Source/ObjectMixer/Public/Views/List/ObjectMixerEditorListRowData.h
+++ b/Engine/Plugins/Editor/ObjectMixer/ObjectMixer/Source/ObjectMixer/Public/Views/List/ObjectMixerEditorListRowData.h
@@ -36,34 +36,53 @@ struct OBJECTMIXEREDITOR_API FObjectMixerEditorListRowData
 		FSceneOutlinerTreeItemID RowIdentifier;
 		FName PropertyName = NAME_None;
 		EPropertyValueSetFlags::Type PropertyValueSetFlags = 0;
+
+		bool operator==(const FPropertyPropagationInfo& Other) const
+		{
+			return PropertyName == Other.PropertyName;
+		}
+
+		friend uint32 GetTypeHash (const FPropertyPropagationInfo& Other)
+		{
+			return GetTypeHash(Other.PropertyName);
+		}
 	};
 	
 	FObjectMixerEditorListRowData(
 		SSceneOutliner* InSceneOutliner, const FText& InDisplayNameOverride = FText::GetEmpty())
-	: SceneOutlinerPtr(InSceneOutliner)
-	, DisplayNameOverride(InDisplayNameOverride)
-	{}
+	: DisplayNameOverride(InDisplayNameOverride)
+	{
+		if (InSceneOutliner)
+		{
+			SceneOutlinerPtr = StaticCastSharedRef<SSceneOutliner>(InSceneOutliner->AsShared());
+		}
+	}
 	
 	FObjectMixerEditorListRowData() = default;
 
 	~FObjectMixerEditorListRowData() = default;
 
+	bool IsValid() const
+	{
+		return SceneOutlinerPtr.IsValid();
+	}
+
 	const TArray<TObjectPtr<UObjectMixerObjectFilter>>& GetObjectFilterInstances() const;
 
 	const UObjectMixerObjectFilter* GetMainObjectFilterInstance() const;
 
-	UE_NODISCARD bool GetIsTreeViewItemExpanded(const TSharedRef<ISceneOutlinerTreeItem> InRow);
+	[[nodiscard]] bool GetIsTreeViewItemExpanded(const TSharedRef<ISceneOutlinerTreeItem> InRow);
 	void SetIsTreeViewItemExpanded(const TSharedRef<ISceneOutlinerTreeItem> InRow, const bool bNewExpanded);
 
-	UE_NODISCARD bool GetIsSelected(const TSharedRef<ISceneOutlinerTreeItem> InRow);
+	[[nodiscard]] bool GetIsSelected(const TSharedRef<ISceneOutlinerTreeItem> InRow);
 	void SetIsSelected(const TSharedRef<ISceneOutlinerTreeItem> InRow, const bool bNewSelected);
 
-	UE_NODISCARD bool HasAtLeastOneChildThatIsNotSolo(
+	[[nodiscard]] bool HasAtLeastOneChildThatIsNotSolo(
 		const TSharedRef<ISceneOutlinerTreeItem> InRow, const bool bRecursive = true) const;
 
-	UE_NODISCARD FText GetDisplayName(TSharedPtr<ISceneOutlinerTreeItem> InTreeItem) const;
+	[[nodiscard]] FText GetDisplayName(TSharedPtr<ISceneOutlinerTreeItem> InTreeItem) const;
 
-	UE_NODISCARD const FText& GetDisplayNameOverride() const
+	[[nodiscard]] const FText& GetDisplayNameOverride() const
 	{
 		return DisplayNameOverride;
 	}
@@ -73,13 +92,13 @@ struct OBJECTMIXEREDITOR_API FObjectMixerEditorListRowData
 		DisplayNameOverride = InDisplayNameOverride;
 	}
 
-	UE_NODISCARD SObjectMixerEditorList* GetListView() const;
+	[[nodiscard]] TWeakPtr<SObjectMixerEditorList> GetListView() const;
 
-	UE_NODISCARD TArray<TSharedPtr<ISceneOutlinerTreeItem>> GetSelectedTreeViewItems() const;
+	[[nodiscard]] TArray<TSharedPtr<ISceneOutlinerTreeItem>> GetSelectedTreeViewItems() const;
 	
 	void OnChangeVisibility(const FSceneOutlinerTreeItemRef TreeItem, const bool bNewVisible);
 	
-	UE_NODISCARD const FTransientEditorVisibilityRules& GetVisibilityRules() const;
+	[[nodiscard]] const FTransientEditorVisibilityRules& GetVisibilityRules() const;
 	void SetVisibilityRules(const FTransientEditorVisibilityRules& InVisibilityRules);
 
 	bool IsUserSetHiddenInEditor() const;
@@ -92,31 +111,39 @@ struct OBJECTMIXEREDITOR_API FObjectMixerEditorListRowData
 
 	bool GetIsHybridRow() const
 	{
-		return HybridComponent.IsValid();
+		return HybridComponentPath.IsValid();
 	}
 
 	UActorComponent* GetHybridComponent() const
 	{
-		return HybridComponent.Get();
+		return HybridComponentPath.Get();
 	}
 	
 	/** If this row represents an actor or other container and should show the data for a single child component, define it here. */
 	void SetHybridComponent(UActorComponent* InHybridComponent)
 	{
-		HybridComponent = InHybridComponent;
+		HybridComponentPath = InHybridComponent;
 	}
 
-	void PropagateChangesToSimilarSelectedRowProperties(
+	/** Attempt to propagate a property's new value to other selected rows.
+	 * @return Returns true if the property's handle was found or there is no work to do (i.e., no property name specified or this row is not selected)
+	 */
+	bool PropagateChangesToSimilarSelectedRowProperties(
 		const TSharedRef<ISceneOutlinerTreeItem> InRow, const FPropertyPropagationInfo PropertyPropagationInfo);
 	
 	TMap<FName, TWeakPtr<IPropertyHandle>> PropertyNamesToHandles;
 
-	SSceneOutliner* SceneOutlinerPtr;
+	TWeakPtr<SSceneOutliner> SceneOutlinerPtr;
 
 protected:
 	FTransientEditorVisibilityRules VisibilityRules;
 
 	FText DisplayNameOverride;
 
-	TWeakObjectPtr<UActorComponent> HybridComponent = nullptr;
+	/**
+	 * A breadcrumb trail to a single component that matches the selected Object Mixer Filters.
+	 * For example, if your filter is looking for lights and you have an actor with a single light component,
+	 * this will show the actor and component on the same line rather than requiring the user to unfold the actor row to see one component.
+	 */
+	TSoftObjectPtr<UActorComponent> HybridComponentPath = nullptr;
 };
diff --git a/Engine/Plugins/Editor/ObjectMixer/ObjectMixer/Source/ObjectMixer/Public/Views/List/RowTypes/ObjectMixerEditorListRowActor.h b/Engine/Plugins/Editor/ObjectMixer/ObjectMixer/Source/ObjectMixer/Public/Views/List/RowTypes/ObjectMixerEditorListRowActor.h
--- a/Engine/Plugins/Editor/ObjectMixer/ObjectMixer/Source/ObjectMixer/Public/Views/List/RowTypes/ObjectMixerEditorListRowActor.h
+++ b/Engine/Plugins/Editor/ObjectMixer/ObjectMixer/Source/ObjectMixer/Public/Views/List/RowTypes/ObjectMixerEditorListRowActor.h
@@ -21,11 +21,15 @@ struct OBJECTMIXEREDITOR_API FObjectMixerEditorListRowActor : FActorTreeItem
 		AActor* InObject, 
 		SSceneOutliner* InSceneOutliner, const FText& InDisplayNameOverride = FText::GetEmpty())
 	: FActorTreeItem(InObject)
+	, OriginalObjectSoftPtr(InObject)
 	{
 		RowData = FObjectMixerEditorListRowData(InSceneOutliner, InDisplayNameOverride);
 	}
 	
 	FObjectMixerEditorListRowData RowData;
+	
+	/** Used in scenarios where the original object may be reconstructed or trashed, such as when running a construction script. */
+	TSoftObjectPtr<AActor> OriginalObjectSoftPtr;
 
 	/* Begin ISceneOutlinerTreeItem Implementation */
 	static const FSceneOutlinerTreeItemType Type;
diff --git a/Engine/Plugins/Editor/ObjectMixer/ObjectMixer/Source/ObjectMixer/Public/Views/List/RowTypes/ObjectMixerEditorListRowComponent.h b/Engine/Plugins/Editor/ObjectMixer/ObjectMixer/Source/ObjectMixer/Public/Views/List/RowTypes/ObjectMixerEditorListRowComponent.h
--- a/Engine/Plugins/Editor/ObjectMixer/ObjectMixer/Source/ObjectMixer/Public/Views/List/RowTypes/ObjectMixerEditorListRowComponent.h
+++ b/Engine/Plugins/Editor/ObjectMixer/ObjectMixer/Source/ObjectMixer/Public/Views/List/RowTypes/ObjectMixerEditorListRowComponent.h
@@ -21,21 +21,15 @@ struct OBJECTMIXEREDITOR_API FObjectMixerEditorListRowComponent : FComponentTree
 		UActorComponent* InObject, 
 		SSceneOutliner* InSceneOutliner, const FText& InDisplayNameOverride = FText::GetEmpty())
 	: FComponentTreeItem(InObject)
+	, OriginalObjectSoftPtr(InObject)
 	{
 		RowData = FObjectMixerEditorListRowData(InSceneOutliner, InDisplayNameOverride);
 	}
 	
 	FObjectMixerEditorListRowData RowData;
-
-	AActor* GetOwningActor() const
-	{
-		if (const UActorComponent* WrappedComponent = Component.Get())
-		{
-			return WrappedComponent->GetOwner();
-		}
-
-		return nullptr;
-	}
+	
+	/** Used in scenarios where the original object may be reconstructed or trashed, such as when running a construction script. */
+	TSoftObjectPtr<UActorComponent> OriginalObjectSoftPtr;
 	
 	/* Begin ISceneOutlinerTreeItem Implementation */
 	static const FSceneOutlinerTreeItemType Type;
diff --git a/Engine/Plugins/Editor/ObjectMixer/ObjectMixer/Source/ObjectMixer/Public/Views/List/RowTypes/ObjectMixerEditorListRowUObject.h b/Engine/Plugins/Editor/ObjectMixer/ObjectMixer/Source/ObjectMixer/Public/Views/List/RowTypes/ObjectMixerEditorListRowUObject.h
--- a/Engine/Plugins/Editor/ObjectMixer/ObjectMixer/Source/ObjectMixer/Public/Views/List/RowTypes/ObjectMixerEditorListRowUObject.h
+++ b/Engine/Plugins/Editor/ObjectMixer/ObjectMixer/Source/ObjectMixer/Public/Views/List/RowTypes/ObjectMixerEditorListRowUObject.h
@@ -21,7 +21,7 @@ struct OBJECTMIXEREDITOR_API FObjectMixerEditorListRowUObject : ISceneOutlinerTr
 		UObject* InObject, 
 		SSceneOutliner* InSceneOutliner, const FText& InDisplayNameOverride = FText::GetEmpty())
 	: ISceneOutlinerTreeItem(Type)
-	, ObjectPtr(InObject)
+	, ObjectSoftPtr(InObject)
 	, ID(InObject)
 	{
 		RowData = FObjectMixerEditorListRowData(InSceneOutliner, InDisplayNameOverride);
@@ -29,7 +29,8 @@ struct OBJECTMIXEREDITOR_API FObjectMixerEditorListRowUObject : ISceneOutlinerTr
 	
 	FObjectMixerEditorListRowData RowData;
 
-	TObjectPtr<UObject> ObjectPtr;
+	/** Used in scenarios where the original object may be reconstructed or trashed, such as when running a construction script. */
+	TSoftObjectPtr<UObject> ObjectSoftPtr;
 	
 	/** Constant identifier for this tree item */
 	const FObjectKey ID;
@@ -38,7 +39,7 @@ struct OBJECTMIXEREDITOR_API FObjectMixerEditorListRowUObject : ISceneOutlinerTr
 	static const FSceneOutlinerTreeItemType Type;
 	virtual FString GetDisplayString() const override;
 	virtual bool CanInteract() const override { return true; }
-	virtual bool IsValid() const override { return ObjectPtr != nullptr; }
+	virtual bool IsValid() const override { return !ObjectSoftPtr.IsNull(); }
 	virtual FSceneOutlinerTreeItemID GetID() const override { return ID; }
 	/* End ISceneOutlinerTreeItem Implementation */
 };
diff --git a/Engine/Plugins/Editor/ObjectMixer/ObjectMixer/Source/ObjectMixer/Public/Views/List/SObjectMixerEditorList.h b/Engine/Plugins/Editor/ObjectMixer/ObjectMixer/Source/ObjectMixer/Public/Views/List/SObjectMixerEditorList.h
--- a/Engine/Plugins/Editor/ObjectMixer/ObjectMixer/Source/ObjectMixer/Public/Views/List/SObjectMixerEditorList.h
+++ b/Engine/Plugins/Editor/ObjectMixer/ObjectMixer/Source/ObjectMixer/Public/Views/List/SObjectMixerEditorList.h
@@ -176,16 +176,13 @@ protected:
 	/** A reference to the struct that controls this widget */
 	TWeakPtr<FObjectMixerEditorList> ListModelPtr;
 
-	TSharedPtr<SSearchBox> SearchBoxPtr;
-	TSharedPtr<SComboButton> ViewOptionsComboButton;
-
 	// User Collections
 	
 	TSharedPtr<SWrapBox> CollectionSelectorBox;
 
 	bool bIsRebuildRequested = false;
 	
-	TArray<FObjectMixerEditorListRowData::FPropertyPropagationInfo> PendingPropertyPropagations;
+	TSet<FObjectMixerEditorListRowData::FPropertyPropagationInfo> PendingPropertyPropagations;
 
 	TArray<FObjectMixerSceneOutlinerColumnInfo> HeaderColumnInfos;
 
