Object Mixer:

-Fix: Alt key not inverting selection sync

diff --git a/Engine/Plugins/Editor/ObjectMixer/ObjectMixer/Source/ObjectMixer/Private/Views/List/Modes/ObjectMixerOutlinerMode.cpp b/Engine/Plugins/Editor/ObjectMixer/ObjectMixer/Source/ObjectMixer/Private/Views/List/Modes/ObjectMixerOutlinerMode.cpp
--- a/Engine/Plugins/Editor/ObjectMixer/ObjectMixer/Source/ObjectMixer/Private/Views/List/Modes/ObjectMixerOutlinerMode.cpp
+++ b/Engine/Plugins/Editor/ObjectMixer/ObjectMixer/Source/ObjectMixer/Private/Views/List/Modes/ObjectMixerOutlinerMode.cpp
@@ -568,10 +568,9 @@ void FObjectMixerOutlinerMode::OnLevelSelectionChanged(UObject* Obj)
 	// If the SceneOutliner's reentrant flag is set, the selection change has already been handled in the outliner class
 	if (!SceneOutliner->GetIsReentrant())
 	{
-		SceneOutliner->ClearSelection();
 		SceneOutliner->RefreshSelection();
 
-		// Scroll last item into view  - this means if we are multi-selecting, we show newest selection. @TODO Not perfect though
+		// Scroll last item into view - this means if we are multi-selecting, we show newest selection. @TODO Not perfect though
 		if (const AActor* LastSelectedActor = GEditor->GetSelectedActors()->GetBottom<AActor>())
 		{
 			if (FSceneOutlinerTreeItemPtr TreeItem = SceneOutliner->GetTreeItem(LastSelectedActor, false))
@@ -624,6 +623,56 @@ void FObjectMixerOutlinerMode::OnPostLoadMapWithWorld(UWorld* World)
 	SceneOutliner->FullRefresh();
 }
 
+bool EvaluateSyncSelectionSetting()
+{
+	const bool bIsSyncSettingOn = GetDefault<UObjectMixerEditorSettings>()->bSyncSelection;
+	const bool bIsAltDown = FSlateApplication::Get().GetModifierKeys().IsAltDown();
+
+	// Alt key should invert the user setting
+	return (bIsSyncSettingOn && !bIsAltDown) || (!bIsSyncSettingOn && bIsAltDown);
+}
+
+bool FObjectMixerOutlinerMode::ShouldSyncSelectionToEditor()
+{
+	return bShouldTemporarilyForceSelectionSyncToEditor || EvaluateSyncSelectionSetting();
+}
+
+bool FObjectMixerOutlinerMode::ShouldSyncSelectionFromEditor()
+{
+	return bShouldTemporarilyForceSelectionSyncFromEditor || EvaluateSyncSelectionSetting();
+}
+
+void FObjectMixerOutlinerMode::SynchronizeAllSelectionsToEditor()
+{
+	bool bAreAnyInPIE = false;
+
+	// Actors
+	SynchronizeSelectedActorDescs();
+	TArray<AActor*> SelectedActors = SceneOutliner->GetSelection().GetData<AActor*>(ObjectMixerOutliner::FActorSelector());
+	bool bHasActorSelectionChanged = HasActorSelectionChanged(SelectedActors, bAreAnyInPIE);
+		
+	// Components
+	TArray<UActorComponent*> SelectedComponents = SceneOutliner->GetSelection().GetData<UActorComponent*>(ObjectMixerOutliner::FComponentSelector());
+	bool bHasComponentSelectionChanged = HasComponentSelectionChanged(SelectedComponents, bAreAnyInPIE);
+		
+	// If there's a discrepancy, update the selected objects to reflect the list.
+	if (bHasActorSelectionChanged || bHasComponentSelectionChanged)
+	{
+		const bool bShouldActuallyTransact = !bAreAnyInPIE;
+		const FScopedTransaction Transaction(LOCTEXT("ClickingOnActorsAndComponents", "Clicking on Actors & Components"), bShouldActuallyTransact);
+
+		if (bHasActorSelectionChanged)
+		{
+			SelectActorsInEditor(SelectedActors);
+		}
+			
+		if (bHasComponentSelectionChanged)
+		{
+			SelectComponentsInEditor(SelectedComponents);
+		}
+	}
+}
+
 bool FObjectMixerOutlinerMode::HasActorSelectionChanged(TArray<AActor*>& OutSelectedActors, bool& bOutAreAnyInPIE)
 {
 	bool bHasSelectionChanged = false;
@@ -663,7 +712,6 @@ bool FObjectMixerOutlinerMode::HasActorSelectionChanged(TArray<AActor*>& OutSele
 				{
 					OutSelectedActors.Remove(GroupActor);
 				}
-
 			}
 		}
 	}
@@ -1791,37 +1839,12 @@ void FObjectMixerOutlinerMode::OnActorDescRemoved(FWorldPartitionActorDesc* InAc
 
 void FObjectMixerOutlinerMode::OnItemSelectionChanged(FSceneOutlinerTreeItemPtr TreeItem, ESelectInfo::Type SelectionType, const FSceneOutlinerItemSelection& Selection)
 {
-	if (GetDefault<UObjectMixerEditorSettings>()->bSyncSelection)
+	if (ShouldSyncSelectionToEditor())
 	{
-		bool bAreAnyInPIE = false;
-
-		// Actors
-		SynchronizeSelectedActorDescs();
-		TArray<AActor*> SelectedActors = Selection.GetData<AActor*>(ObjectMixerOutliner::FActorSelector());
-		bool bHasActorSelectionChanged = HasActorSelectionChanged(SelectedActors, bAreAnyInPIE);
-		
-		// Components
-		TArray<UActorComponent*> SelectedComponents = Selection.GetData<UActorComponent*>(ObjectMixerOutliner::FComponentSelector());
-		bool bHasComponentSelectionChanged = HasComponentSelectionChanged(SelectedComponents, bAreAnyInPIE);
-		
-		// If there's a discrepancy, update the selected objects to reflect the list.
-		if (bHasActorSelectionChanged || bHasComponentSelectionChanged)
-		{
-			const FScopedTransaction Transaction(LOCTEXT("ClickingOnActorsAndComponents", "Clicking on Actors & Components"), !bAreAnyInPIE);
-
-			if (bHasActorSelectionChanged)
-			{
-				SelectActorsInEditor(SelectedActors);
-			}
-			
-			if (bHasComponentSelectionChanged)
-			{
-				SelectComponentsInEditor(SelectedComponents);
-			}
-
-			SceneOutliner->RefreshSelection();
-		}
+		SynchronizeAllSelectionsToEditor();
 	}
+	
+	bShouldTemporarilyForceSelectionSyncToEditor = false;
 }
 
 void FObjectMixerOutlinerMode::OnItemDoubleClick(FSceneOutlinerTreeItemPtr Item)
@@ -3114,12 +3137,14 @@ void FObjectMixerOutlinerMode::UnpinItems(const TArray<FSceneOutlinerTreeItemPtr
 
 void FObjectMixerOutlinerMode::SynchronizeSelection()
 {
-	if (GetDefault<UObjectMixerEditorSettings>()->bSyncSelection)
+	if (ShouldSyncSelectionFromEditor())
 	{
 		SynchronizeActorSelection();
 		SynchronizeComponentSelection();
 		SynchronizeSelectedActorDescs();
 	}
+	
+	bShouldTemporarilyForceSelectionSyncFromEditor = false;
 }
 
 FCreateSceneOutlinerMode FObjectMixerOutlinerMode::CreateFolderPickerMode(const FFolder::FRootObject& InRootObject) const
diff --git a/Engine/Plugins/Editor/ObjectMixer/ObjectMixer/Source/ObjectMixer/Public/Views/List/Modes/ObjectMixerOutlinerMode.h b/Engine/Plugins/Editor/ObjectMixer/ObjectMixer/Source/ObjectMixer/Public/Views/List/Modes/ObjectMixerOutlinerMode.h
--- a/Engine/Plugins/Editor/ObjectMixer/ObjectMixer/Source/ObjectMixer/Public/Views/List/Modes/ObjectMixerOutlinerMode.h
+++ b/Engine/Plugins/Editor/ObjectMixer/ObjectMixer/Source/ObjectMixer/Public/Views/List/Modes/ObjectMixerOutlinerMode.h
@@ -279,6 +279,9 @@ protected:
 	void OnLevelActorRequestsRename(const AActor* Actor);
 	void OnPostLoadMapWithWorld(UWorld* World);
 
+	bool ShouldSyncSelectionToEditor();
+	bool ShouldSyncSelectionFromEditor();
+	void SynchronizeAllSelectionsToEditor();
 	bool HasActorSelectionChanged(TArray<AActor*>& OutSelectedActors, bool& bOutAreAnyInPIE);
 	bool HasComponentSelectionChanged(TArray<UActorComponent*>& OutSelectedComponents, bool& bOutAreAnyInPIE);
 	static void SelectActorsInEditor(const TArray<AActor*>& InSelectedActors);
@@ -347,4 +350,9 @@ protected:
 	
 	TArray<FFilterClassSelectionInfo> FilterClassSelectionInfos;
 	TSharedRef<SWidget> OnGenerateFilterClassMenu();
+
+	/** Used in case of a selection sync override. */
+	
+	bool bShouldTemporarilyForceSelectionSyncFromEditor = false;
+	bool bShouldTemporarilyForceSelectionSyncToEditor = false;
 };
