Fix ObjectMixerOutlinerMode duplicated code from FActorBrowsingMode.

diff --git a/Engine/Plugins/Editor/ObjectMixer/ObjectMixer/Source/ObjectMixer/Private/Views/List/Modes/ObjectMixerOutlinerMode.cpp b/Engine/Plugins/Editor/ObjectMixer/ObjectMixer/Source/ObjectMixer/Private/Views/List/Modes/ObjectMixerOutlinerMode.cpp
--- a/Engine/Plugins/Editor/ObjectMixer/ObjectMixer/Source/ObjectMixer/Private/Views/List/Modes/ObjectMixerOutlinerMode.cpp
+++ b/Engine/Plugins/Editor/ObjectMixer/ObjectMixer/Source/ObjectMixer/Private/Views/List/Modes/ObjectMixerOutlinerMode.cpp
@@ -2083,15 +2083,12 @@ void FObjectMixerOutlinerMode::RepairErrors() const
 			ISourceControlModule& SourceControlModule = ISourceControlModule::Get();
 			ISourceControlProvider& SourceControlProvider = SourceControlModule.GetProvider();
 
-			TArray<FWorldPartitionActorDesc*> InvalidActorDescs;
-			WorldPartition->ForEachActorDescContainer([&InvalidActorDescs](UActorDescContainer* ActorDescContainer)
+			TArray<FAssetData> InvalidActorAssets;
+			WorldPartition->ForEachActorDescContainer([&InvalidActorAssets](UActorDescContainer* ActorDescContainer)
 			{
-				for (const TUniquePtr<FWorldPartitionActorDesc>& InvalidActor : ActorDescContainer->GetInvalidActors())
+				for (const FAssetData& InvalidActor : ActorDescContainer->GetInvalidActors())
 				{
-					if (FWorldPartitionActorDesc* InvalidActorDesc = InvalidActor.Get())
-					{
-						InvalidActorDescs.Add(InvalidActorDesc);
-					}
+					InvalidActorAssets.Add(InvalidActor);
 				}
 				ActorDescContainer->ClearInvalidActors();
 			});
@@ -2099,13 +2096,13 @@ void FObjectMixerOutlinerMode::RepairErrors() const
 			TArray<FString> ActorFilesToDelete;
 			TArray<FString> ActorFilesToRevert;
 			{
-				FScopedSlowTask SlowTask(InvalidActorDescs.Num(), LOCTEXT("UpdatingSourceControlStatus", "Updating revision control status..."));
+				FScopedSlowTask SlowTask(InvalidActorAssets.Num(), LOCTEXT("UpdatingSourceControlStatus", "Updating source control status..."));
 				SlowTask.MakeDialogDelayed(1.0f);
 
-				for (const FWorldPartitionActorDesc* InvalidActorDesc : InvalidActorDescs)
+				for (const FAssetData& InvalidActorAsset : InvalidActorAssets)
 				{
 					FPackagePath PackagePath;
-					if (FPackagePath::TryFromPackageName(InvalidActorDesc->GetActorPackage(), PackagePath))
+					if (FPackagePath::TryFromPackageName(InvalidActorAsset.PackageName, PackagePath))
 					{
 						const FString ActorFile = PackagePath.GetLocalFullPath();
 						FSourceControlStatePtr SCState = SourceControlProvider.GetState(ActorFile, EStateCacheUsage::ForceUpdate);
