Object Mixer:

-Fix: Custom Row Types don't get reported correctly by the scene outliner
-Utils script cleanup

diff --git a/Engine/Plugins/Editor/ObjectMixer/ObjectMixer/Source/ObjectMixer/Private/Views/List/ObjectMixerUtils.cpp b/Engine/Plugins/Editor/ObjectMixer/ObjectMixer/Source/ObjectMixer/Private/Views/List/ObjectMixerUtils.cpp
--- a/Engine/Plugins/Editor/ObjectMixer/ObjectMixer/Source/ObjectMixer/Private/Views/List/ObjectMixerUtils.cpp
+++ b/Engine/Plugins/Editor/ObjectMixer/ObjectMixer/Source/ObjectMixer/Private/Views/List/ObjectMixerUtils.cpp
@@ -11,66 +11,26 @@
 
 FObjectMixerEditorListRowFolder* FObjectMixerUtils::AsFolderRow(TSharedPtr<ISceneOutlinerTreeItem> InTreeItem)
 {
-	// Sometimes the outliner reports type as the parent type, but we can be assured we're always using our type.
-	// For this reason, we check for our type first and if it's not our type, check parent type then cast.
-	if (FObjectMixerEditorListRowFolder* AsFolder = InTreeItem->CastTo<FObjectMixerEditorListRowFolder>())
-	{
-		return AsFolder;
-	}
-	
-	if (FActorFolderTreeItem* AsFolder = InTreeItem->CastTo<FActorFolderTreeItem>())
-	{
-		if (AsFolder->GetFolder().IsValid())
-		{
-			return StaticCast<FObjectMixerEditorListRowFolder*>(AsFolder);
-		}
-	}
-
-	return nullptr;
+	check (InTreeItem);
+	return InTreeItem->CastTo<FObjectMixerEditorListRowFolder>();
 }
 
 FObjectMixerEditorListRowActor* FObjectMixerUtils::AsActorRow(TSharedPtr<ISceneOutlinerTreeItem> InTreeItem)
 {
-	// Sometimes the outliner reports type as the parent type, but we can be assured we're always using our type.
-	// For this reason, we check for our type first and if it's not our type, check parent type then cast.
-	if (FObjectMixerEditorListRowActor* AsActor = InTreeItem->CastTo<FObjectMixerEditorListRowActor>())
-	{
-		return AsActor;
-	}
-	
-	if (FActorTreeItem* AsActor = InTreeItem->CastTo<FActorTreeItem>())
-	{
-		return StaticCast<FObjectMixerEditorListRowActor*>(AsActor);
-	}
-
-	return nullptr;
+	check (InTreeItem);
+	return InTreeItem->CastTo<FObjectMixerEditorListRowActor>();
 }
 
 FObjectMixerEditorListRowComponent* FObjectMixerUtils::AsComponentRow(TSharedPtr<ISceneOutlinerTreeItem> InTreeItem)
 {
-	// Sometimes the outliner reports type as the parent type, but we can be assured we're always using our type.
-	// For this reason, we check for our type first and if it's not our type, check parent type then cast.
-	if (FObjectMixerEditorListRowComponent* AsComponent = InTreeItem->CastTo<FObjectMixerEditorListRowComponent>())
-	{
-		return AsComponent;
-	}
-	
-	if (FComponentTreeItem* AsComponent = InTreeItem->CastTo<FComponentTreeItem>())
-	{
-		return StaticCast<FObjectMixerEditorListRowComponent*>(AsComponent);
-	}
-
-	return nullptr;
+	check (InTreeItem);
+	return InTreeItem->CastTo<FObjectMixerEditorListRowComponent>();
 }
 
 FObjectMixerEditorListRowUObject* FObjectMixerUtils::AsObjectRow(TSharedPtr<ISceneOutlinerTreeItem> InTreeItem)
 {
-	if (FObjectMixerEditorListRowUObject* AsObject = InTreeItem->CastTo<FObjectMixerEditorListRowUObject>())
-	{
-		return AsObject;
-	}
-
-	return nullptr;
+	check (InTreeItem);
+	return InTreeItem->CastTo<FObjectMixerEditorListRowUObject>();
 }
 
 FObjectMixerEditorListRowData* FObjectMixerUtils::GetRowData(TSharedPtr<ISceneOutlinerTreeItem> InTreeItem)
@@ -139,7 +99,7 @@ UObject* FObjectMixerUtils::GetRowObject(TSharedPtr<ISceneOutlinerTreeItem> InTr
 
 	if (const FObjectMixerEditorListRowUObject* ObjectRow = AsObjectRow(InTreeItem))
 	{
-		ObjectRow->ObjectSoftPtr.Get();
+		return ObjectRow->ObjectSoftPtr.Get();
 	}
 
 	return nullptr;
diff --git a/Engine/Plugins/Editor/ObjectMixer/ObjectMixer/Source/ObjectMixer/Public/Views/List/RowTypes/ObjectMixerEditorListRowActor.h b/Engine/Plugins/Editor/ObjectMixer/ObjectMixer/Source/ObjectMixer/Public/Views/List/RowTypes/ObjectMixerEditorListRowActor.h
--- a/Engine/Plugins/Editor/ObjectMixer/ObjectMixer/Source/ObjectMixer/Public/Views/List/RowTypes/ObjectMixerEditorListRowActor.h
+++ b/Engine/Plugins/Editor/ObjectMixer/ObjectMixer/Source/ObjectMixer/Public/Views/List/RowTypes/ObjectMixerEditorListRowActor.h
@@ -5,15 +5,8 @@
 #include "Views/List/ObjectMixerEditorListRowData.h"
 
 #include "ActorTreeItem.h"
-#include "Containers/Array.h"
-#include "Containers/UnrealString.h"
-#include "Folder.h"
 #include "GameFramework/Actor.h"
-#include "Layout/Visibility.h"
-#include "PropertyHandle.h"
 #include "SSceneOutliner.h"
-#include "Templates/SharedPointer.h"
-#include "UObject/Object.h"
 
 struct OBJECTMIXEREDITOR_API FObjectMixerEditorListRowActor : FActorTreeItem
 {
@@ -23,6 +16,7 @@ struct OBJECTMIXEREDITOR_API FObjectMixerEditorListRowActor : FActorTreeItem
 	: FActorTreeItem(InObject)
 	, OriginalObjectSoftPtr(InObject)
 	{
+		TreeType = Type;
 		RowData = FObjectMixerEditorListRowData(InSceneOutliner, InDisplayNameOverride);
 	}
 	
diff --git a/Engine/Plugins/Editor/ObjectMixer/ObjectMixer/Source/ObjectMixer/Public/Views/List/RowTypes/ObjectMixerEditorListRowComponent.h b/Engine/Plugins/Editor/ObjectMixer/ObjectMixer/Source/ObjectMixer/Public/Views/List/RowTypes/ObjectMixerEditorListRowComponent.h
--- a/Engine/Plugins/Editor/ObjectMixer/ObjectMixer/Source/ObjectMixer/Public/Views/List/RowTypes/ObjectMixerEditorListRowComponent.h
+++ b/Engine/Plugins/Editor/ObjectMixer/ObjectMixer/Source/ObjectMixer/Public/Views/List/RowTypes/ObjectMixerEditorListRowComponent.h
@@ -5,15 +5,7 @@
 #include "Views/List/ObjectMixerEditorListRowData.h"
 
 #include "ComponentTreeItem.h"
-#include "Containers/Array.h"
-#include "Containers/UnrealString.h"
-#include "Folder.h"
-#include "GameFramework/Actor.h"
-#include "Layout/Visibility.h"
-#include "PropertyHandle.h"
 #include "SSceneOutliner.h"
-#include "Templates/SharedPointer.h"
-#include "UObject/Object.h"
 
 struct OBJECTMIXEREDITOR_API FObjectMixerEditorListRowComponent : FComponentTreeItem
 {
@@ -23,6 +15,7 @@ struct OBJECTMIXEREDITOR_API FObjectMixerEditorListRowComponent : FComponentTree
 	: FComponentTreeItem(InObject)
 	, OriginalObjectSoftPtr(InObject)
 	{
+		TreeType = Type;
 		RowData = FObjectMixerEditorListRowData(InSceneOutliner, InDisplayNameOverride);
 	}
 	
diff --git a/Engine/Plugins/Editor/ObjectMixer/ObjectMixer/Source/ObjectMixer/Public/Views/List/RowTypes/ObjectMixerEditorListRowFolder.h b/Engine/Plugins/Editor/ObjectMixer/ObjectMixer/Source/ObjectMixer/Public/Views/List/RowTypes/ObjectMixerEditorListRowFolder.h
--- a/Engine/Plugins/Editor/ObjectMixer/ObjectMixer/Source/ObjectMixer/Public/Views/List/RowTypes/ObjectMixerEditorListRowFolder.h
+++ b/Engine/Plugins/Editor/ObjectMixer/ObjectMixer/Source/ObjectMixer/Public/Views/List/RowTypes/ObjectMixerEditorListRowFolder.h
@@ -6,15 +6,8 @@
 
 #include "ActorFolderTreeItem.h"
 
-#include "Containers/Array.h"
-#include "Containers/UnrealString.h"
 #include "Folder.h"
-#include "GameFramework/Actor.h"
-#include "Layout/Visibility.h"
-#include "PropertyHandle.h"
 #include "SSceneOutliner.h"
-#include "Templates/SharedPointer.h"
-#include "UObject/Object.h"
 
 struct OBJECTMIXEREDITOR_API FObjectMixerEditorListRowFolder : FActorFolderTreeItem
 {
@@ -23,6 +16,7 @@ struct OBJECTMIXEREDITOR_API FObjectMixerEditorListRowFolder : FActorFolderTreeI
 		SSceneOutliner* InSceneOutliner, UWorld* World, const FText& InDisplayNameOverride = FText::GetEmpty())
 	: FActorFolderTreeItem(InFolder, World)
 	{
+		TreeType = Type;
 		RowData = FObjectMixerEditorListRowData(InSceneOutliner, InDisplayNameOverride);
 	}
 	
diff --git a/Engine/Plugins/Editor/ObjectMixer/ObjectMixer/Source/ObjectMixer/Public/Views/List/RowTypes/ObjectMixerEditorListRowUObject.h b/Engine/Plugins/Editor/ObjectMixer/ObjectMixer/Source/ObjectMixer/Public/Views/List/RowTypes/ObjectMixerEditorListRowUObject.h
--- a/Engine/Plugins/Editor/ObjectMixer/ObjectMixer/Source/ObjectMixer/Public/Views/List/RowTypes/ObjectMixerEditorListRowUObject.h
+++ b/Engine/Plugins/Editor/ObjectMixer/ObjectMixer/Source/ObjectMixer/Public/Views/List/RowTypes/ObjectMixerEditorListRowUObject.h
@@ -4,15 +4,9 @@
 
 #include "Views/List/ObjectMixerEditorListRowData.h"
 
-#include "Containers/Array.h"
 #include "Containers/UnrealString.h"
-#include "Folder.h"
-#include "GameFramework/Actor.h"
 #include "ISceneOutlinerTreeItem.h"
-#include "Layout/Visibility.h"
-#include "PropertyHandle.h"
 #include "SSceneOutliner.h"
-#include "Templates/SharedPointer.h"
 #include "UObject/Object.h"
 
 struct OBJECTMIXEREDITOR_API FObjectMixerEditorListRowUObject : ISceneOutlinerTreeItem
@@ -24,6 +18,7 @@ struct OBJECTMIXEREDITOR_API FObjectMixerEditorListRowUObject : ISceneOutlinerTr
 	, ObjectSoftPtr(InObject)
 	, ID(InObject)
 	{
+		TreeType = Type;
 		RowData = FObjectMixerEditorListRowData(InSceneOutliner, InDisplayNameOverride);
 	}
 	
