Fix deprecation

diff --git a/Engine/Plugins/Editor/ObjectMixer/ObjectMixer/Source/ObjectMixer/Private/Views/List/SObjectMixerEditorList.cpp b/Engine/Plugins/Editor/ObjectMixer/ObjectMixer/Source/ObjectMixer/Private/Views/List/SObjectMixerEditorList.cpp
--- a/Engine/Plugins/Editor/ObjectMixer/ObjectMixer/Source/ObjectMixer/Private/Views/List/SObjectMixerEditorList.cpp
+++ b/Engine/Plugins/Editor/ObjectMixer/ObjectMixer/Source/ObjectMixer/Private/Views/List/SObjectMixerEditorList.cpp
@@ -944,9 +944,9 @@ void SObjectMixerEditorList::CreateActorTextInfoColumns(FSceneOutlinerInitializa
 		{
 			if (const FWorldPartitionActorDesc* ActorDesc = ActorDescItem->ActorDescHandle.Get())
 			{
-				if (FName LevelPackage = ActorDesc->GetLevelPackage(); !LevelPackage.IsNone())
+				if (FName LevelPackage = ActorDesc->GetContainerPackage(); !LevelPackage.IsNone())
 				{
-					return ActorDesc->GetLevelPackage().ToString();
+					return ActorDesc->GetContainerPackage().ToString();
 				}
 			}
 		}
