fix duplicated code compilation

diff --git a/Engine/Plugins/Editor/ObjectMixer/ObjectMixer/Source/ObjectMixer/Private/Views/List/Modes/ObjectMixerOutlinerHierarchy.cpp b/Engine/Plugins/Editor/ObjectMixer/ObjectMixer/Source/ObjectMixer/Private/Views/List/Modes/ObjectMixerOutlinerHierarchy.cpp
--- a/Engine/Plugins/Editor/ObjectMixer/ObjectMixer/Source/ObjectMixer/Private/Views/List/Modes/ObjectMixerOutlinerHierarchy.cpp
+++ b/Engine/Plugins/Editor/ObjectMixer/ObjectMixer/Source/ObjectMixer/Private/Views/List/Modes/ObjectMixerOutlinerHierarchy.cpp
@@ -126,13 +126,10 @@ FObjectMixerOutlinerHierarchy::~FObjectMixerOutlinerHierarchy()
 	FWorldDelegates::LevelRemovedFromWorld.RemoveAll(this);
 
 
-	if (FActorFolders::IsAvailable())
-	{
-		auto& Folders = FActorFolders::Get();
-		Folders.OnFolderCreated.RemoveAll(this);
-		Folders.OnFolderMoved.RemoveAll(this);
-		Folders.OnFolderDeleted.RemoveAll(this);
-	}
+	auto& Folders = FActorFolders::Get();
+	Folders.OnFolderCreated.RemoveAll(this);
+	Folders.OnFolderMoved.RemoveAll(this);
+	Folders.OnFolderDeleted.RemoveAll(this);
 }
 
 FSceneOutlinerTreeItemPtr FObjectMixerOutlinerHierarchy::FindOrCreateParentItem(
