Object Mixer:

-Show Spawnable Lights From Sequencer
-Fix: Object Mixer missing properties in Object Mixer Window
-Add "Restore Default Columns" option

$jira UE-183379

diff --git a/Engine/Plugins/Editor/ObjectMixer/LightMixer/Source/LightMixer/Public/LightMixerObjectFilter.h b/Engine/Plugins/Editor/ObjectMixer/LightMixer/Source/LightMixer/Public/LightMixerObjectFilter.h
--- a/Engine/Plugins/Editor/ObjectMixer/LightMixer/Source/LightMixer/Public/LightMixerObjectFilter.h
+++ b/Engine/Plugins/Editor/ObjectMixer/LightMixer/Source/LightMixer/Public/LightMixerObjectFilter.h
@@ -31,6 +31,11 @@ public:
 		};
 	}
 
+	virtual bool GetShowTransientObjects() const override
+	{
+		return true;
+	}
+
 	virtual TSet<FName> GetColumnsToShowByDefault() const override
 	{
 		return 
diff --git a/Engine/Plugins/Editor/ObjectMixer/ObjectMixer/Source/ObjectMixer/ObjectMixerEditor.build.cs b/Engine/Plugins/Editor/ObjectMixer/ObjectMixer/Source/ObjectMixer/ObjectMixerEditor.build.cs
--- a/Engine/Plugins/Editor/ObjectMixer/ObjectMixer/Source/ObjectMixer/ObjectMixerEditor.build.cs
+++ b/Engine/Plugins/Editor/ObjectMixer/ObjectMixer/Source/ObjectMixer/ObjectMixerEditor.build.cs
@@ -29,6 +29,7 @@ public class ObjectMixerEditor : ModuleRules
 		PrivateDependencyModuleNames.AddRange(
 			new string[]
 			{
+				"AdvancedWidgets",
 				"ApplicationCore",
 				"AssetRegistry",
 				"AssetTools",
@@ -41,8 +42,10 @@ public class ObjectMixerEditor : ModuleRules
 				"InputCore",
 				"Kismet",
 				"LevelEditor",
+				"LevelSequence",
 				"Projects",
 				"SceneOutliner",
+				"Sequencer",
 				"Slate",
 				"SlateCore",
 				"ToolMenus", 
diff --git a/Engine/Plugins/Editor/ObjectMixer/ObjectMixer/Source/ObjectMixer/Private/ObjectFilter/ObjectMixerEditorObjectFilter.cpp b/Engine/Plugins/Editor/ObjectMixer/ObjectMixer/Source/ObjectMixer/Private/ObjectFilter/ObjectMixerEditorObjectFilter.cpp
--- a/Engine/Plugins/Editor/ObjectMixer/ObjectMixer/Source/ObjectMixer/Private/ObjectFilter/ObjectMixerEditorObjectFilter.cpp
+++ b/Engine/Plugins/Editor/ObjectMixer/ObjectMixer/Source/ObjectMixer/Private/ObjectFilter/ObjectMixerEditorObjectFilter.cpp
@@ -74,6 +74,11 @@ FText UObjectMixerObjectFilter::GetRowTooltipText(UObject* InObject, const bool
 	return FText::GetEmpty();
 }
 
+bool UObjectMixerObjectFilter::GetShowTransientObjects() const
+{
+	return false;
+}
+
 bool UObjectMixerObjectFilter::GetRowEditorVisibility(UObject* InObject) const
 {
 	if (IsValid(InObject))
diff --git a/Engine/Plugins/Editor/ObjectMixer/ObjectMixer/Source/ObjectMixer/Private/ObjectMixerEditorModule.cpp b/Engine/Plugins/Editor/ObjectMixer/ObjectMixer/Source/ObjectMixer/Private/ObjectMixerEditorModule.cpp
--- a/Engine/Plugins/Editor/ObjectMixer/ObjectMixer/Source/ObjectMixer/Private/ObjectMixerEditorModule.cpp
+++ b/Engine/Plugins/Editor/ObjectMixer/ObjectMixer/Source/ObjectMixer/Private/ObjectMixerEditorModule.cpp
@@ -3,18 +3,19 @@
 #include "ObjectMixerEditorModule.h"
 
 #include "ObjectMixerEditorLog.h"
-#include "Engine/Level.h"
 #include "ObjectMixerEditorSettings.h"
-#include "Misc/CoreDelegates.h"
 #include "ObjectMixerEditorSerializedData.h"
 #include "ObjectMixerEditorStyle.h"
 #include "Views/List/ObjectMixerEditorList.h"
 #include "Views/Widgets/ObjectMixerEditorListMenuContext.h"
 
+#include "Engine/Level.h"
 #include "ISettingsModule.h"
 #include "LevelEditor.h"
-#include "Selection.h"
+#include "LevelEditorSequencerIntegration.h"
+#include "Misc/CoreDelegates.h"
 #include "Misc/TransactionObjectEvent.h"
+#include "Selection.h"
 #include "ToolMenus.h"
 #include "Widgets/Docking/SDockTab.h"
 #include "WorkspaceMenuStructure.h"
@@ -93,6 +94,8 @@ void FObjectMixerEditorModule::Teardown()
 		LevelEditorModule->OnComponentsEdited().RemoveAll(this);
 	}
 
+	FLevelEditorSequencerIntegration::Get().GetOnSequencersChanged().RemoveAll(this);
+
 	for (FDelegateHandle& Delegate : DelegateHandles)
 	{
 		Delegate.Reset();
@@ -137,11 +140,20 @@ TSharedPtr<SWidget> FObjectMixerEditorModule::MakeObjectMixerDialog(
 		ListModel->SetDefaultFilterClass(InDefaultFilterClass);
 	}
 	
-	const TSharedPtr<SWidget> ObjectMixerDialog = ListModel->GetOrCreateWidget();
+	const TSharedPtr<SWidget> ObjectMixerDialog = ListModel->CreateWidget();
 
 	return ObjectMixerDialog;
 }
 
+void FObjectMixerEditorModule::RegenerateListWidget()
+{
+	if (DockTab)
+	{
+		const TSharedPtr<SWidget> ObjectMixerDialog = MakeObjectMixerDialog(DefaultFilterClass);
+		DockTab->SetContent(ObjectMixerDialog ? ObjectMixerDialog.ToSharedRef() : SNullWidget::NullWidget);
+	}
+}
+
 void FObjectMixerEditorModule::RequestRebuildList() const
 {
 	if (ListModel.IsValid())
@@ -265,16 +277,14 @@ void FObjectMixerEditorModule::UnregisterSettings() const
 
 TSharedRef<SDockTab> FObjectMixerEditorModule::SpawnTab(const FSpawnTabArgs& Args)
 {
-	const TSharedRef<SDockTab> DockTab =
-		SNew(SDockTab)
+	SAssignNew(DockTab, SDockTab)
 		.Label(TabLabel)
 		.TabRole(ETabRole::NomadTab)
 	;
 
-	const TSharedPtr<SWidget> ObjectMixerDialog = MakeObjectMixerDialog(DefaultFilterClass);
-	DockTab->SetContent(ObjectMixerDialog ? ObjectMixerDialog.ToSharedRef() : SNullWidget::NullWidget);
+	RegenerateListWidget();
 			
-	return DockTab;
+	return DockTab.ToSharedRef();
 }
 
 TSharedPtr<FWorkspaceItem> FObjectMixerEditorModule::GetWorkspaceGroup()
@@ -312,7 +322,10 @@ void FObjectMixerEditorModule::BindDelegates()
 
 	DelegateHandles.Add(GEditor->OnBlueprintCompiled().AddLambda([this] ()
 	{
-		RequestRebuildList();
+		if (ListModel)
+		{
+			ListModel->CacheAndRebuildFilters(true);
+		}
 	}));
 	DelegateHandles.Add(GEditor->OnLevelActorAdded().AddLambda([this] (AActor*)
 	{
@@ -357,10 +370,16 @@ void FObjectMixerEditorModule::BindDelegates()
 	}));
 	
 	FLevelEditorModule& LevelEditorModule = FModuleManager::LoadModuleChecked<FLevelEditorModule>("LevelEditor");
-	LevelEditorModule.OnComponentsEdited().AddLambda([this]()
+	DelegateHandles.Add(LevelEditorModule.OnComponentsEdited().AddLambda([this]()
 	{
 		RequestRebuildList();
-	});
+	}));
+	
+	DelegateHandles.Add(FLevelEditorSequencerIntegration::Get().GetOnSequencersChanged().AddLambda(
+		[this]
+		{
+			RegenerateListWidget();
+		}));
 }
 
 TSet<FName> FObjectMixerEditorModule::GetPropertiesThatRequireRefresh() const
diff --git a/Engine/Plugins/Editor/ObjectMixer/ObjectMixer/Source/ObjectMixer/Private/Views/List/Modes/ObjectMixerOutlinerHierarchy.cpp b/Engine/Plugins/Editor/ObjectMixer/ObjectMixer/Source/ObjectMixer/Private/Views/List/Modes/ObjectMixerOutlinerHierarchy.cpp
--- a/Engine/Plugins/Editor/ObjectMixer/ObjectMixer/Source/ObjectMixer/Private/Views/List/Modes/ObjectMixerOutlinerHierarchy.cpp
+++ b/Engine/Plugins/Editor/ObjectMixer/ObjectMixer/Source/ObjectMixer/Private/Views/List/Modes/ObjectMixerOutlinerHierarchy.cpp
@@ -673,7 +673,10 @@ void FObjectMixerOutlinerHierarchy::CreateChildren(const FSceneOutlinerTreeItemP
 	else if (const FActorTreeItem* ParentActorItem = Item->CastTo<FActorTreeItem>())
 	{
 		AActor* ParentActor = ParentActorItem->Actor.Get();
-		check(ParentActor->GetWorld() == RepresentingWorld);
+		if (ParentActor->GetWorld() != RepresentingWorld)
+		{
+			return;
+		}
 
 		CreateComponentItems(ParentActor, OutChildren);
 
diff --git a/Engine/Plugins/Editor/ObjectMixer/ObjectMixer/Source/ObjectMixer/Private/Views/List/Modes/ObjectMixerOutlinerMode.cpp b/Engine/Plugins/Editor/ObjectMixer/ObjectMixer/Source/ObjectMixer/Private/Views/List/Modes/ObjectMixerOutlinerMode.cpp
--- a/Engine/Plugins/Editor/ObjectMixer/ObjectMixer/Source/ObjectMixer/Private/Views/List/Modes/ObjectMixerOutlinerMode.cpp
+++ b/Engine/Plugins/Editor/ObjectMixer/ObjectMixer/Source/ObjectMixer/Private/Views/List/Modes/ObjectMixerOutlinerMode.cpp
@@ -9,6 +9,7 @@
 #include "Views/List/Modes/ObjectMixerOutlinerHierarchy.h"
 #include "Views/List/Modes/SFilterClassMenuItem.h"
 #include "Views/List/ObjectMixerUtils.h"
+#include "Views/List/RowTypes/ObjectMixerEditorListRowActor.h"
 #include "Views/Widgets/ObjectMixerEditorListMenuContext.h"
 
 #include "ActorDescTreeItem.h"
@@ -454,7 +455,7 @@ FObjectMixerOutlinerMode::FObjectMixerOutlinerMode(
 		{
 			const TSet<FName>& SelectedCollections = ListModel->GetSelectedCollections();
 
-			if (SelectedCollections.Num() == 0)
+			if (SelectedCollections.Num() == 0 || (SelectedCollections.Num() == 1 && SelectedCollections.Find("All")))
 			{
 				// "All" collection selected
 				return true;
@@ -806,7 +807,8 @@ void FObjectMixerOutlinerMode::OnActorLabelChanged(AActor* ChangedActor)
 	if (IsActorDisplayable(ChangedActor) && RepresentingWorld.Get() == ChangedActor->GetWorld())
 	{
 		// Force create the item otherwise the outliner may not be notified of a change to the item if it is filtered out
-		if (FSceneOutlinerTreeItemPtr Item = CreateItemFor<FActorTreeItem>(ChangedActor, true))
+		if (FSceneOutlinerTreeItemPtr Item = CreateItemFor<FObjectMixerEditorListRowActor>(
+			FObjectMixerEditorListRowActor(ChangedActor, GetSceneOutliner()), true))
 		{
 			SceneOutliner->OnItemLabelChanged(Item);
 		}
@@ -820,7 +822,8 @@ void FObjectMixerOutlinerMode::OnObjectsReplaced(const TMap<UObject*, UObject*>&
 		AActor* Actor = Cast<AActor>(Pair.Value);
 		if (Actor && RepresentingWorld.Get() == Actor->GetWorld() && IsActorDisplayable(Actor))
 		{
-			if (FSceneOutlinerTreeItemPtr Item = CreateItemFor<FActorTreeItem>(Actor, true))
+			if (FSceneOutlinerTreeItemPtr Item = CreateItemFor<FObjectMixerEditorListRowActor>(
+				FObjectMixerEditorListRowActor(Actor, GetSceneOutliner()), true))
 			{
 				SceneOutliner->OnItemLabelChanged(Item);
 			}
@@ -1164,7 +1167,7 @@ TSharedRef<SWidget> FObjectMixerOutlinerMode::OnGenerateFilterClassMenu()
 							}
 						}
 
-						PinnedList->CacheAndRebuildFilters();
+						PinnedList->CacheAndRebuildFilters(true);
 					}
 
 					return FReply::Handled();
diff --git a/Engine/Plugins/Editor/ObjectMixer/ObjectMixer/Source/ObjectMixer/Private/Views/List/ObjectMixerEditorList.cpp b/Engine/Plugins/Editor/ObjectMixer/ObjectMixer/Source/ObjectMixer/Private/Views/List/ObjectMixerEditorList.cpp
--- a/Engine/Plugins/Editor/ObjectMixer/ObjectMixer/Source/ObjectMixer/Private/Views/List/ObjectMixerEditorList.cpp
+++ b/Engine/Plugins/Editor/ObjectMixer/ObjectMixer/Source/ObjectMixer/Private/Views/List/ObjectMixerEditorList.cpp
@@ -3,11 +3,13 @@
 #include "Views/List/ObjectMixerEditorList.h"
 
 #include "ObjectMixerEditorLog.h"
+#include "ObjectMixerEditorModule.h"
 #include "ObjectMixerEditorSerializedData.h"
-#include "Views/List/SObjectMixerEditorList.h"
 
-#include "LevelEditorActions.h"
 #include "Framework/Commands/GenericCommands.h"
+#include "LevelEditorActions.h"
+#include "Modules/ModuleManager.h"
+#include "Views/List/SObjectMixerEditorList.h"
 
 void FObjectMixerEditorList::Initialize()
 {
@@ -136,6 +138,11 @@ void FObjectMixerEditorList::RemoveObjectFilterClass(UClass* InObjectFilterClass
 	}
 }
 
+FObjectMixerEditorModule* FObjectMixerEditorList::GetModulePtr() const
+{
+	return FModuleManager::GetModulePtr<FObjectMixerEditorModule>(GetModuleName());
+}
+
 UObjectMixerEditorSerializedData* FObjectMixerEditorList::GetSerializedData() const
 {
 	return GetMutableDefault<UObjectMixerEditorSerializedData>();
@@ -334,7 +341,7 @@ TSharedRef<SWidget> FObjectMixerEditorList::GetOrCreateWidget()
 {
 	if (!ListWidget.IsValid())
 	{
-		SAssignNew(ListWidget, SObjectMixerEditorList, SharedThis(this));
+		CreateWidget();
 	}
 
 	RequestRebuildList();
@@ -342,6 +349,19 @@ TSharedRef<SWidget> FObjectMixerEditorList::GetOrCreateWidget()
 	return ListWidget.ToSharedRef();
 }
 
+TSharedRef<SWidget> FObjectMixerEditorList::CreateWidget()
+{
+	return SAssignNew(ListWidget, SObjectMixerEditorList, SharedThis(this));
+}
+
+void FObjectMixerEditorList::RequestRegenerateListWidget()
+{
+	if (FObjectMixerEditorModule* Module = GetModulePtr())
+	{
+		Module->RegenerateListWidget();
+	}
+}
+
 void FObjectMixerEditorList::OnPostFilterChange()
 {
 	RequestRebuildList();
@@ -369,7 +389,11 @@ void FObjectMixerEditorList::BuildPerformanceCache()
 	TSet<FName> LocalColumnsToShowByDefaultCache;
 	TSet<FName> LocalColumnsToExcludeCache;
 	TSet<FName> LocalForceAddedColumnsCache;
-	
+
+	bShouldIncludeUnsupportedPropertiesCache = false;
+	bShouldShowTransientObjectsCache = false;
+
+	// Aggregate caches - append everything to all-inclusive lists
 	for (const TObjectPtr<UObjectMixerObjectFilter>& FilterInstance : GetObjectFilterInstances())
 	{
 		if (!FilterInstance)
@@ -382,6 +406,16 @@ void FObjectMixerEditorList::BuildPerformanceCache()
 		LocalColumnsToShowByDefaultCache.Append(FilterInstance->GetColumnsToShowByDefault());
 		LocalColumnsToExcludeCache.Append(FilterInstance->GetColumnsToExclude());
 		LocalForceAddedColumnsCache.Append(FilterInstance->GetForceAddedColumns());
+
+		// These are zero-sum - if they're true for any filter, they must remain true
+		if (!bShouldIncludeUnsupportedPropertiesCache)
+		{
+			bShouldIncludeUnsupportedPropertiesCache = FilterInstance->ShouldIncludeUnsupportedProperties();
+		}
+		if (!bShouldShowTransientObjectsCache)
+		{
+			bShouldShowTransientObjectsCache = FilterInstance->GetShowTransientObjects();
+		}
 	}
 
 	if (LocalObjectClassesToFilterCache.Difference(ObjectClassesToFilterCache).Num() > 0 || 
@@ -417,11 +451,6 @@ void FObjectMixerEditorList::BuildPerformanceCache()
 		{
 			PropertyInheritanceInclusionOptionsCache = PropertyInheritanceInclusionOptions;
 		}
-		if (const bool bShouldIncludeUnsupportedProperties = FilterInstance->ShouldIncludeUnsupportedProperties(); 
-			bShouldIncludeUnsupportedProperties != bShouldIncludeUnsupportedPropertiesCache)
-		{
-			bShouldIncludeUnsupportedPropertiesCache = bShouldIncludeUnsupportedProperties;
-		}
 	}
 	else
 	{
@@ -429,6 +458,11 @@ void FObjectMixerEditorList::BuildPerformanceCache()
 	}
 }
 
+bool FObjectMixerEditorList::ShouldShowTransientObjects() const
+{
+	return bShouldShowTransientObjectsCache;
+}
+
 void FObjectMixerEditorList::OnRenameCommand()
 {
 	if (ListWidget.IsValid())
diff --git a/Engine/Plugins/Editor/ObjectMixer/ObjectMixer/Source/ObjectMixer/Private/Views/List/SObjectMixerEditorList.cpp b/Engine/Plugins/Editor/ObjectMixer/ObjectMixer/Source/ObjectMixer/Private/Views/List/SObjectMixerEditorList.cpp
--- a/Engine/Plugins/Editor/ObjectMixer/ObjectMixer/Source/ObjectMixer/Private/Views/List/SObjectMixerEditorList.cpp
+++ b/Engine/Plugins/Editor/ObjectMixer/ObjectMixer/Source/ObjectMixer/Private/Views/List/SObjectMixerEditorList.cpp
@@ -14,24 +14,23 @@
 #include "Views/Widgets/SCollectionSelectionButton.h"
 #include "Views/Widgets/SObjectMixerPlacementAssetMenuEntry.h"
 
-// View Filters
 #include "ActorDescTreeItem.h"
 #include "ActorFolderTreeItem.h"
 #include "ComponentTreeItem.h"
 #include "Editor.h"
 #include "EditorActorFolders.h"
 #include "IPlacementModeModule.h"
+#include "LevelEditorSequencerIntegration.h"
+#include "LevelSequence.h"
 #include "SceneOutlinerModule.h"
 #include "SceneOutlinerTextInfoColumn.h"
 #include "SPositiveActionButton.h"
 #include "Algo/Find.h"
-#include "Engine/Level.h"
+#include "Algo/RemoveIf.h"
 #include "Framework/Application/SlateApplication.h"
 #include "Framework/MultiBox/MultiBoxBuilder.h"
-#include "Framework/Views/TableViewMetadata.h"
 #include "LevelInstance/LevelInstanceInterface.h"
 #include "LevelInstance/LevelInstanceSubsystem.h"
-#include "ToolMenus.h"
 #include "Widgets/Input/SCheckBox.h"
 #include "Widgets/Layout/SWrapBox.h"
 #include "Widgets/Views/STreeView.h"
@@ -47,6 +46,14 @@ const FName SObjectMixerEditorList::ItemNameColumnName(TEXT("Builtin_Name"));
 const FName SObjectMixerEditorList::EditorVisibilityColumnName(TEXT("Builtin_EditorVisibility"));
 const FName SObjectMixerEditorList::EditorVisibilitySoloColumnName(TEXT("Builtin_EditorVisibilitySolo"));
 
+const FName SequencerSpawnableColumnName("Spawnable");
+const FName SequencerInfoColumnName("Sequence");
+
+void CloseAllMenus()
+{
+	FSlateApplication::Get().DismissAllMenus();
+}
+
 void SObjectMixerEditorList::Construct(const FArguments& InArgs, TSharedRef<FObjectMixerEditorList> ListModel)
 {
 	ListModelPtr = ListModel;
@@ -56,6 +63,7 @@ void SObjectMixerEditorList::Construct(const FArguments& InArgs, TSharedRef<FObj
 	InitOptions.bFocusSearchBoxWhenOpened = true;
 	InitOptions.bShowCreateNewFolder = true;
 	InitOptions.bShowParentTree = true;
+	InitOptions.bShowTransient = ListModel->ShouldShowTransientObjects();
 	InitOptions.OutlinerIdentifier = ListModel->GetModuleName();
 	
 	// Mode
@@ -73,15 +81,15 @@ void SObjectMixerEditorList::Construct(const FArguments& InArgs, TSharedRef<FObj
 	// Header Row
 	InitOptions.bShowHeaderRow = true;
 	InitOptions.bCanSelectGeneratedColumns = false;
-	SetupColumns(InitOptions);
 	 
+	SetupColumns(InitOptions);
 	SSceneOutliner::Construct(InArgs, InitOptions);
 
 	InsertCollectionSelector();
 
 	if (UObjectMixerEditorSerializedData* SerializedData = GetMutableDefault<UObjectMixerEditorSerializedData>())
 	{
-		SerializedData->OnObjectMixerCollectionMapChanged.AddRaw(this, &SObjectMixerEditorList::RebuildCollectionSelector);
+		SerializedData->OnObjectMixerCollectionMapChanged.AddSP(this, &SObjectMixerEditorList::RebuildCollectionSelector);
 	}
 	
 	RebuildCollectionSelector();
@@ -89,6 +97,10 @@ void SObjectMixerEditorList::Construct(const FArguments& InArgs, TSharedRef<FObj
 	SetSingleCollectionSelection();
 }
 
+SObjectMixerEditorList::~SObjectMixerEditorList()
+{	
+	FLevelEditorSequencerIntegration::Get().GetOnSequencersChanged().RemoveAll(this);
+}
 
 FReply SObjectMixerEditorList::OnKeyDown(const FGeometry& MyGeometry, const FKeyEvent& InKeyEvent)
 {
@@ -582,6 +594,14 @@ FObjectMixerSceneOutlinerColumnInfo* SObjectMixerEditorList::GetColumnInfoByProp
 		});
 }
 
+void SObjectMixerEditorList::RestoreDefaultPropertyColumns()
+{
+	for (const FObjectMixerSceneOutlinerColumnInfo& ColumnInfo : HeaderColumnInfos)
+	{
+		SetColumnVisibility(ColumnInfo.ColumnID, ColumnInfo.bIsDesiredToBeShownByDefault);
+	}
+}
+
 void SObjectMixerEditorList::CustomAddToToolbar(TSharedPtr<SHorizontalBox> Toolbar)
 {
 	ToolbarPtr = Toolbar;
@@ -653,12 +673,29 @@ void SObjectMixerEditorList::Tick(const FGeometry& AllottedGeometry, const doubl
 	}
 }
 
-TSharedRef<SWidget> SObjectMixerEditorList::GenerateHeaderRowContextMenu() const
+TSharedRef<SWidget> SObjectMixerEditorList::GenerateHeaderRowContextMenu()
 {
 	FMenuBuilder MenuBuilder(false, nullptr);
 	
 	MenuBuilder.AddSearchWidget();
 
+	MenuBuilder.BeginSection(TEXT("Actions"), LOCTEXT("ContextMenuActions", "Actions"));
+	{
+		MenuBuilder.AddMenuEntry(
+			LOCTEXT("RestoreDefaultPropertyColumns", "Restore Default Columns"),
+			LOCTEXT("RestoreDefaultPropertyColumnsTooltip", "Restore property columns to built-in columns and those as defined in GetColumnsToShowByDefault() in your object filters."),
+			FSlateIcon(),
+			FUIAction(
+				FExecuteAction::CreateLambda([this]()
+				{
+					CloseAllMenus();
+					RestoreDefaultPropertyColumns();
+				})
+			)
+		);
+	}
+	MenuBuilder.EndSection();
+
 	FName LastPropertyCategoryName = NAME_None;
 
 	for (const FObjectMixerSceneOutlinerColumnInfo& ColumnInfo : HeaderColumnInfos)
@@ -696,15 +733,6 @@ TSharedRef<SWidget> SObjectMixerEditorList::GenerateHeaderRowContextMenu() const
 					const bool bNewColumnEnabled = !PinnedHeaderRow->IsColumnVisible(ColumnID);
 					
 					PinnedHeaderRow->SetShowGeneratedColumn(ColumnID, bNewColumnEnabled);
-
-					if (UObjectMixerEditorSerializedData* SerializedData = GetMutableDefault<UObjectMixerEditorSerializedData>())
-					{
-						for (const TSubclassOf<UObjectMixerObjectFilter>& FilterClass : PinnedListModel->GetObjectFilterClasses())
-						{
-							const FName FilterName = FilterClass->GetFName();
-							SerializedData->SetShouldShowColumn(FilterName, ColumnID, bNewColumnEnabled);
-						}
-					}
 				}),
 				FCanExecuteAction::CreateLambda([bCanSelectColumn](){return bCanSelectColumn;}),
 				FIsActionChecked::CreateLambda([this, ColumnID]()
@@ -778,7 +806,8 @@ bool SObjectMixerEditorList::AddUniquePropertyColumnInfo(
 					Property->GetDisplayNameText(),
 					EListViewColumnType::PropertyGenerated,
 					*PropertyCategoryName,
-					true, true, false, false
+					true,
+					bForceIncludeProperty ? true : GetListModelPtr().Pin()->ColumnsToShowByDefaultCache.Contains(PropertyName)
 				}
 			);
 
@@ -1048,20 +1077,6 @@ void SObjectMixerEditorList::CreateActorTextInfoColumns(FSceneOutlinerInitializa
 	auto AddTextInfoColumn = [this, PinnedListModel, &OutInitOptions](FName ColumnID, TAttribute<FText> ColumnName, FGetTextForItem ColumnInfo)
 	{
 		bool bShouldShowColumn = false;
-		// Load visible columns from SerializedData
-		if (UObjectMixerEditorSerializedData* SerializedData = GetMutableDefault<UObjectMixerEditorSerializedData>())
-		{
-			for (const TSubclassOf<UObjectMixerObjectFilter>& Filter : PinnedListModel->GetObjectFilterClasses())
-			{
-				const FName FilterName = Filter->GetFName();
-				if (SerializedData->IsColumnDataSerialized(FilterName, ColumnID))
-				{
-					// If any filter's serialized data says we should show a column, show it and stomp other filters' saved setting.
-					bShouldShowColumn = SerializedData->ShouldShowColumn(FilterName, ColumnID);
-					break;
-				}
-			}
-		}
 		
 		OutInitOptions.ColumnMap.Add(
 			ColumnID,
@@ -1088,7 +1103,7 @@ void SObjectMixerEditorList::CreateActorTextInfoColumns(FSceneOutlinerInitializa
 				ColumnName.Get(FText::GetEmpty()),
 				EListViewColumnType::BuiltIn,
 				TEXT("Built-In"),
-				false, true, true, false
+				true, bShouldShowColumn
 			}
 		);
 	};
@@ -1116,26 +1131,9 @@ void SObjectMixerEditorList::SetupColumns(FSceneOutlinerInitializationOptions& O
 			ESceneOutlinerColumnVisibility InDefaultVisibility, int32 InPriorityIndex, const FName ColumnID,
 			const FText DisplayText, const bool bCanBeHidden, FCreateSceneOutlinerColumn CreationDelegate = FCreateSceneOutlinerColumn(),
 			TOptional<float> FillSize = TOptional<float>())
-		{
-			bool bShouldShowColumn = InDefaultVisibility == ESceneOutlinerColumnVisibility::Visible ? true : false;
-			
-			// Load visible columns from SerializedData
-			if (UObjectMixerEditorSerializedData* SerializedData = GetMutableDefault<UObjectMixerEditorSerializedData>())
-			{
-				for (const TSubclassOf<UObjectMixerObjectFilter>& Filter : PinnedListModel->GetObjectFilterClasses())
-				{
-					const FName FilterName = Filter->GetFName();
-					if (SerializedData->IsColumnDataSerialized(FilterName, ColumnID))
-					{
-						// If any filter's serialized data says we should show a column, show it and stomp other filters' saved setting.
-						bShouldShowColumn = SerializedData->ShouldShowColumn(FilterName, ColumnID);
-						break;
-					}
-				}
-			}
-			
+		{			
 			OutInitOptions.ColumnMap.Add(ColumnID, FSceneOutlinerColumnInfo(
-				bShouldShowColumn ? ESceneOutlinerColumnVisibility::Visible : ESceneOutlinerColumnVisibility::Invisible,
+				InDefaultVisibility,
 				InPriorityIndex, CreationDelegate, bCanBeHidden,
 				FillSize, DisplayText, EHeaderComboVisibility::Never,
 				FOnGetContent::CreateRaw(this, &SObjectMixerEditorList::GenerateHeaderRowContextMenu)));
@@ -1146,7 +1144,8 @@ void SObjectMixerEditorList::SetupColumns(FSceneOutlinerInitializationOptions& O
 					DisplayText,
 					EListViewColumnType::BuiltIn,
 					TEXT("Built-In"),
-					bShouldShowColumn, bCanBeHidden, true, false
+					bCanBeHidden,
+					InDefaultVisibility == ESceneOutlinerColumnVisibility::Visible ? true : false
 				}
 			);
 		};
@@ -1180,11 +1179,25 @@ void SObjectMixerEditorList::SetupColumns(FSceneOutlinerInitializationOptions& O
 			ESceneOutlinerColumnVisibility::Visible, 11, FSceneOutlinerBuiltInColumnTypes::ActorInfo(),
 			FSceneOutlinerBuiltInColumnTypes::ActorInfo_Localized(), true			
 		);
+		
+		// Sequencer Columns
+		if (FLevelEditorSequencerIntegration::Get().GetSequencers().Num() > 0)
+		{
+			const FText SpawnableDisplayText = LOCTEXT("SpawnableColumnName", "Spawnable");
+			const FText SequencerInfoDisplayText = LOCTEXT("SequencerColumnName", "Sequencer");
 
-		AddBuiltInColumn(
-			ESceneOutlinerColumnVisibility::Invisible, 12, FSceneOutlinerBuiltInColumnTypes::SourceControl(),
-			FSceneOutlinerBuiltInColumnTypes::SourceControl_Localized(), true			
-		);
+			AddBuiltInColumn(
+				ESceneOutlinerColumnVisibility::Visible, 12, SequencerSpawnableColumnName,
+				SpawnableDisplayText, true,
+				FCreateSceneOutlinerColumn::CreateRaw( &FLevelEditorSequencerIntegration::Get(), &FLevelEditorSequencerIntegration::CreateSequencerSpawnableColumn)
+			);
+
+			AddBuiltInColumn(
+				ESceneOutlinerColumnVisibility::Visible, 13, SequencerInfoColumnName,
+				SequencerInfoDisplayText, true,
+				FCreateSceneOutlinerColumn::CreateRaw( &FLevelEditorSequencerIntegration::Get(), &FLevelEditorSequencerIntegration::CreateSequencerInfoColumn)
+			);
+		}
 
 		ESceneOutlinerColumnVisibility UnsavedColumnVisibility = ESceneOutlinerColumnVisibility::Invisible;
 
@@ -1244,23 +1257,7 @@ void SObjectMixerEditorList::SetupColumns(FSceneOutlinerInitializationOptions& O
 		}
 	}
 
-	// Alphabetical sort by Property Name
-	HeaderColumnInfos.StableSort([](const FObjectMixerSceneOutlinerColumnInfo& A, const FObjectMixerSceneOutlinerColumnInfo& B)
-	{
-		return A.PropertyDisplayText.ToString() < B.PropertyDisplayText.ToString();
-	});
-
-	// Alphabetical sort by Property Category Name
-	HeaderColumnInfos.StableSort([](const FObjectMixerSceneOutlinerColumnInfo& A, const FObjectMixerSceneOutlinerColumnInfo& B)
-	{
-		return A.PropertyCategoryName.LexicalLess(B.PropertyCategoryName);
-	});
-
-	// sort by type
-	HeaderColumnInfos.StableSort([](const FObjectMixerSceneOutlinerColumnInfo& A, const FObjectMixerSceneOutlinerColumnInfo& B)
-	{
-		return A.PropertyType < B.PropertyType;
-	});
+	SortColumnsForHeaderContextMenu();
 
 	// Add Property Columns to Header
 	{
@@ -1278,39 +1275,17 @@ void SObjectMixerEditorList::SetupColumns(FSceneOutlinerInitializationOptions& O
 			{
 				continue;
 			}
-
-			// Figure out if we should show a certain column
-			// First check filter rules
-			bool bShouldShowColumn = PinnedListModel->ColumnsToShowByDefaultCache.Contains(ColumnInfo.ColumnID);
-
-			// Then load visible columns from SerializedData
-			if (UObjectMixerEditorSerializedData* SerializedData = GetMutableDefault<UObjectMixerEditorSerializedData>())
-			{
-				for (const TSubclassOf<UObjectMixerObjectFilter>& Filter : PinnedListModel->GetObjectFilterClasses())
-				{
-					const FName FilterName = Filter->GetFName();
-					if (SerializedData->IsColumnDataSerialized(FilterName, ColumnInfo.ColumnID))
-					{
-						if (SerializedData->ShouldShowColumn(FilterName, ColumnInfo.ColumnID))
-						{
-							// If any filter's serialized data says we should show a column, show it and stomp other filters' saved setting.
-							bShouldShowColumn = true;
-							break;
-						}
-					}
-				}
-			}
 			
 			OutInitOptions.ColumnMap.Add(
 				ColumnInfo.ColumnID,
 				FSceneOutlinerColumnInfo(
-					bShouldShowColumn ? ESceneOutlinerColumnVisibility::Visible : ESceneOutlinerColumnVisibility::Invisible, 30,
+					ColumnInfo.bIsDesiredToBeShownByDefault ? ESceneOutlinerColumnVisibility::Visible : ESceneOutlinerColumnVisibility::Invisible, 30,
 					FCreateSceneOutlinerColumn::CreateLambda([Property = ColumnInfo.PropertyRef](ISceneOutliner& InSceneOutliner)
 					{
 						return MakeShared<FObjectMixerOutlinerPropertyColumn>(InSceneOutliner, Property);
 					}),
 					true,
-					ColumnInfo.FillWidth,
+					TOptional<float>(),
 					PropertyName,
 					EHeaderComboVisibility::Never,
 					FOnGetContent::CreateRaw(this, &SObjectMixerEditorList::GenerateHeaderRowContextMenu)
@@ -1320,6 +1295,27 @@ void SObjectMixerEditorList::SetupColumns(FSceneOutlinerInitializationOptions& O
 	}
 }
 
+void SObjectMixerEditorList::SortColumnsForHeaderContextMenu()
+{
+	// Alphabetical sort by Property Name
+	HeaderColumnInfos.StableSort([](const FObjectMixerSceneOutlinerColumnInfo& A, const FObjectMixerSceneOutlinerColumnInfo& B)
+	{
+		return A.PropertyDisplayText.ToString() < B.PropertyDisplayText.ToString();
+	});
+
+	// Alphabetical sort by Property Category Name
+	HeaderColumnInfos.StableSort([](const FObjectMixerSceneOutlinerColumnInfo& A, const FObjectMixerSceneOutlinerColumnInfo& B)
+	{
+		return A.PropertyCategoryName.LexicalLess(B.PropertyCategoryName);
+	});
+
+	// Sort by type
+	HeaderColumnInfos.StableSort([](const FObjectMixerSceneOutlinerColumnInfo& A, const FObjectMixerSceneOutlinerColumnInfo& B)
+	{
+		return A.PropertyType < B.PropertyType;
+	});
+}
+
 void SObjectMixerEditorList::PropagatePropertyChangesToSelectedRows()
 {
 	using LambdaType = void(*)(const TSet<TWeakPtr<ISceneOutlinerTreeItem>>&,
diff --git a/Engine/Plugins/Editor/ObjectMixer/ObjectMixer/Source/ObjectMixer/Public/ObjectFilter/ObjectMixerEditorObjectFilter.h b/Engine/Plugins/Editor/ObjectMixer/ObjectMixer/Source/ObjectMixer/Public/ObjectFilter/ObjectMixerEditorObjectFilter.h
--- a/Engine/Plugins/Editor/ObjectMixer/ObjectMixer/Source/ObjectMixer/Public/ObjectFilter/ObjectMixerEditorObjectFilter.h
+++ b/Engine/Plugins/Editor/ObjectMixer/ObjectMixer/Source/ObjectMixer/Public/ObjectFilter/ObjectMixerEditorObjectFilter.h
@@ -81,6 +81,11 @@ public:
 	 */
 	virtual FText GetRowTooltipText(UObject* InObject, const bool bIsHybridRow = false) const;
 
+	/**
+	 * Determines if transient objects (such as Sequencer Spawnables) should be shown in the list. False by default.
+	 */
+	virtual bool GetShowTransientObjects() const;
+
 	/**
 	 * Controls how to display the row's visibility icon. Return true if the object should be visible.
 	 * Generally this should work like the Scene Outliner does.
@@ -225,6 +230,17 @@ public:
 	{
 		return Super::GetRowTooltipText(InObject, bIsHybridRow);
 	}
+
+	/**
+	 * Determines if transient objects (such as Sequencer Spawnables) should be shown in the list. False by default.
+	 */
+	UFUNCTION(BlueprintNativeEvent, BlueprintCallable, Category = "Object Mixer")
+	bool GetShowTransientObjects() const;
+
+	bool GetShowTransientObjects_Implementation() const
+	{
+		return Super::GetShowTransientObjects();
+	}
 	
 	/**
 	 * Controls how to display the row's visibility icon. Return true if the object should be visible.
diff --git a/Engine/Plugins/Editor/ObjectMixer/ObjectMixer/Source/ObjectMixer/Public/ObjectMixerEditorModule.h b/Engine/Plugins/Editor/ObjectMixer/ObjectMixer/Source/ObjectMixer/Public/ObjectMixerEditorModule.h
--- a/Engine/Plugins/Editor/ObjectMixer/ObjectMixer/Source/ObjectMixer/Public/ObjectMixerEditorModule.h
+++ b/Engine/Plugins/Editor/ObjectMixer/ObjectMixer/Source/ObjectMixer/Public/ObjectMixerEditorModule.h
@@ -31,6 +31,9 @@ public:
 	
 	virtual TSharedPtr<SWidget> MakeObjectMixerDialog(
 		TSubclassOf<UObjectMixerObjectFilter> InDefaultFilterClass = nullptr);
+	
+	/** Build the List widget from scratch. */
+	void RegenerateListWidget();
 
 	/**
 	 * Regenerate the list items and refresh the list. Call when adding or removing variables.
@@ -87,6 +90,8 @@ protected:
 	/** The text that appears on the spawned nomad tab */
 	FText TabLabel;
 
+	TSharedPtr<SDockTab> DockTab;
+
 	/** Menu Item variables */
 	FText MenuItemName;
 	FSlateIcon MenuItemIcon;
diff --git a/Engine/Plugins/Editor/ObjectMixer/ObjectMixer/Source/ObjectMixer/Public/Views/List/ObjectMixerEditorList.h b/Engine/Plugins/Editor/ObjectMixer/ObjectMixer/Source/ObjectMixer/Public/Views/List/ObjectMixerEditorList.h
--- a/Engine/Plugins/Editor/ObjectMixer/ObjectMixer/Source/ObjectMixer/Public/Views/List/ObjectMixerEditorList.h
+++ b/Engine/Plugins/Editor/ObjectMixer/ObjectMixer/Source/ObjectMixer/Public/Views/List/ObjectMixerEditorList.h
@@ -27,9 +27,16 @@ public:
 	
 	void RegisterAndMapContextMenuCommands();
 
+	/** Creates a new widget only if one does not already exist. Otherwise returns the existing widget. */
 	TSharedRef<SWidget> GetOrCreateWidget();
 
-	/** Regenerate the list items and refresh the list. Call when adding or removing items. */
+	/** Creates a new widget, replacing the existing one's pointer. Not useful to call alone, use RequestRegenerateListWidget */
+	TSharedRef<SWidget> CreateWidget();
+
+	/** Calls back to the module to replace the widget in the dock tab. Call when columns or the object filters have changed. */
+	void RequestRegenerateListWidget();
+
+	/** Rebuild the list items and refresh the list. Call when adding or removing items. */
 	void RequestRebuildList() const;
 
 	/**
@@ -40,6 +47,8 @@ public:
 
 	void BuildPerformanceCache();
 
+	bool ShouldShowTransientObjects() const;
+
 	/** Called when the Rename command is executed from the UI or hotkey. */
 	void OnRenameCommand();
 	
@@ -97,9 +106,16 @@ public:
 		}
 	}
 
-	void CacheAndRebuildFilters()
+	void CacheAndRebuildFilters(const bool bShouldRegenerateWidget = false)
 	{
 		CacheObjectFilterInstances();
+
+		if (bShouldRegenerateWidget)
+		{
+			RequestRegenerateListWidget();
+			return;
+		}
+		
 		RequestRebuildList();
 	}
 
@@ -109,6 +125,8 @@ public:
 		return ModuleName;
 	}
 
+	FObjectMixerEditorModule* GetModulePtr() const;
+
 	// User Collections
 
 	/** Get a pointer to the UObjectMixerEditorSerializedData object along with the name of the filter represented by this ListModel instance. */
@@ -168,6 +186,7 @@ public:
 	TSet<FName> ForceAddedColumnsCache;
 	EObjectMixerInheritanceInclusionOptions PropertyInheritanceInclusionOptionsCache = EObjectMixerInheritanceInclusionOptions::None;
 	bool bShouldIncludeUnsupportedPropertiesCache = false;
+	bool bShouldShowTransientObjectsCache = false;
 
 protected:
 	
@@ -193,11 +212,6 @@ protected:
 	 */
 	TArray<TSubclassOf<UObjectMixerObjectFilter>> ObjectFilterClasses;
 
-	/**
-	 * The FObjectMixerEditorModule subclass that created this instance.
-	 */
-	TWeakPtr<FObjectMixerEditorModule> SpawningModulePtr;
-
 	/**
 	 * If set, this is the filter class used to initialize the ListModel.
 	 * This filter class cannot be turned off by the end user.
diff --git a/Engine/Plugins/Editor/ObjectMixer/ObjectMixer/Source/ObjectMixer/Public/Views/List/SObjectMixerEditorList.h b/Engine/Plugins/Editor/ObjectMixer/ObjectMixer/Source/ObjectMixer/Public/Views/List/SObjectMixerEditorList.h
--- a/Engine/Plugins/Editor/ObjectMixer/ObjectMixer/Source/ObjectMixer/Public/Views/List/SObjectMixerEditorList.h
+++ b/Engine/Plugins/Editor/ObjectMixer/ObjectMixer/Source/ObjectMixer/Public/Views/List/SObjectMixerEditorList.h
@@ -29,41 +29,40 @@ USTRUCT()
 struct FObjectMixerSceneOutlinerColumnInfo
 {
 	GENERATED_BODY()
-	
+
+	/** The pointer to the actual FProperty */
 	FProperty* PropertyRef = nullptr;
-	
+
+	/** The FName of the property */
 	UPROPERTY()
 	FName PropertyName = NAME_None;
 
+	/** The column identifier for the property. Often this is teh same as PropertyName. */
 	UPROPERTY()
 	FName ColumnID = NAME_None;
 
+	/** What will be displayed in the column header unless another widget is defined. */
 	UPROPERTY()
 	FText PropertyDisplayText = FText::GetEmpty();
 
+	/** Is this a built-in column, a column displaying a property widget or something else? */
 	UPROPERTY()
 	EListViewColumnType PropertyType = EListViewColumnType::PropertyGenerated;
 
+	/**
+	 * The category that holds the property.
+	 * Used for sorting properties in the context menu, similar to the details view.
+	 */
 	UPROPERTY()
 	FName PropertyCategoryName = NAME_None;
 
-	UPROPERTY()
-	bool bIsDesiredForDisplay = false;
-
+	/** If true, this column can be enabled and disabled by the user. */
 	UPROPERTY()
 	bool bCanBeHidden = true;
 
+	/** If true, this column will be shown in a clean environment or when default properties are restored. */
 	UPROPERTY()
-	bool bCanBeSorted = false;
-
-	UPROPERTY()
-	bool bUseFixedWidth = false;
-
-	UPROPERTY()
-	float FixedWidth = 25.0f;
-
-	UPROPERTY()
-	float FillWidth = 1.0f;
+	bool bIsDesiredToBeShownByDefault = false;
 };
 
 struct FObjectMixerEditorListRowData;
@@ -80,7 +79,7 @@ public:
 
 	void Construct(const FArguments& InArgs, TSharedRef<FObjectMixerEditorList> ListModel);
 
-	virtual ~SObjectMixerEditorList() {}
+	virtual ~SObjectMixerEditorList() override;
 
 	virtual FReply OnKeyDown( const FGeometry& MyGeometry, const FKeyEvent& InKeyEvent ) override;
 
@@ -160,7 +159,8 @@ public:
 	// Columns
 
 	FObjectMixerSceneOutlinerColumnInfo* GetColumnInfoByPropertyName(const FName& InPropertyName);
-	TSharedRef<SWidget> GenerateHeaderRowContextMenu() const;
+	void RestoreDefaultPropertyColumns();
+	TSharedRef<SWidget> GenerateHeaderRowContextMenu();
 
 	/* Begin SSceneOutliner Interface */
 	virtual void CustomAddToToolbar(TSharedPtr<class SHorizontalBox> Toolbar) override;
@@ -217,4 +217,5 @@ protected:
 	FText GetLevelColumnName() const;
 	void CreateActorTextInfoColumns(FSceneOutlinerInitializationOptions& OutInitOptions);
 	void SetupColumns(FSceneOutlinerInitializationOptions& OutInitOptions);
+	void SortColumnsForHeaderContextMenu();
 };
diff --git a/Engine/Source/Editor/Sequencer/Public/LevelEditorSequencerIntegration.h b/Engine/Source/Editor/Sequencer/Public/LevelEditorSequencerIntegration.h
--- a/Engine/Source/Editor/Sequencer/Public/LevelEditorSequencerIntegration.h
+++ b/Engine/Source/Editor/Sequencer/Public/LevelEditorSequencerIntegration.h
@@ -206,6 +206,8 @@ private:
 	};
 	TArray<FSequencerAndOptions> BoundSequencers;
 
+public:
+	
 	TSharedRef< ISceneOutlinerColumn > CreateSequencerInfoColumn( ISceneOutliner& SceneOutliner ) const;
 	TSharedRef< ISceneOutlinerColumn > CreateSequencerSpawnableColumn( ISceneOutliner& SceneOutliner ) const;
 
