Object Mixer / Light Mixer:

-Fix: Can't open Light Mixer tab again after closing it
-Fix: Light Collections Issues

diff --git a/Engine/Plugins/Editor/ObjectMixer/ObjectMixer/Source/ObjectMixer/Private/ObjectMixerEditorModule.cpp b/Engine/Plugins/Editor/ObjectMixer/ObjectMixer/Source/ObjectMixer/Private/ObjectMixerEditorModule.cpp
--- a/Engine/Plugins/Editor/ObjectMixer/ObjectMixer/Source/ObjectMixer/Private/ObjectMixerEditorModule.cpp
+++ b/Engine/Plugins/Editor/ObjectMixer/ObjectMixer/Source/ObjectMixer/Private/ObjectMixerEditorModule.cpp
@@ -145,13 +145,26 @@ TSharedPtr<SWidget> FObjectMixerEditorModule::MakeObjectMixerDialog(
 	return ObjectMixerDialog;
 }
 
-void FObjectMixerEditorModule::RegenerateListWidget()
+TSharedPtr<SDockTab> FObjectMixerEditorModule::FindNomadTab()
 {
-	if (DockTab)
+	if (!DockTab.IsValid())
+	{
+		DockTab = FGlobalTabmanager::Get()->FindExistingLiveTab(GetTabSpawnerId());
+	}
+	
+	return DockTab.Pin();
+}
+
+bool FObjectMixerEditorModule::RegenerateListWidget()
+{
+	if (const TSharedPtr<SDockTab> FoundTab = FindNomadTab())
 	{
 		const TSharedPtr<SWidget> ObjectMixerDialog = MakeObjectMixerDialog(DefaultFilterClass);
-		DockTab->SetContent(ObjectMixerDialog ? ObjectMixerDialog.ToSharedRef() : SNullWidget::NullWidget);
+		FoundTab->SetContent(ObjectMixerDialog ? ObjectMixerDialog.ToSharedRef() : SNullWidget::NullWidget);
+		return true;
 	}
+
+	return false;
 }
 
 void FObjectMixerEditorModule::RequestRebuildList() const
@@ -223,7 +236,8 @@ void FObjectMixerEditorModule::RegisterTabSpawner()
 		.SetIcon(MenuItemIcon)
 		.SetDisplayName(MenuItemName)
 		.SetTooltipText(MenuItemTooltip)
-		.SetMenuType(TabSpawnerType);
+		.SetMenuType(TabSpawnerType)
+	;
 
 	// Always use the base ObjectMixer function call or WorkspaceGroup may be null 
 	if (!FObjectMixerEditorModule::Get().RegisterItemInMenuGroup(BrowserSpawnerEntry))
@@ -277,14 +291,14 @@ void FObjectMixerEditorModule::UnregisterSettings() const
 
 TSharedRef<SDockTab> FObjectMixerEditorModule::SpawnTab(const FSpawnTabArgs& Args)
 {
-	SAssignNew(DockTab, SDockTab)
+	TSharedRef<SDockTab> NewDockTab = SAssignNew(DockTab, SDockTab)
 		.Label(TabLabel)
 		.TabRole(ETabRole::NomadTab)
 	;
 
 	RegenerateListWidget();
 			
-	return DockTab.ToSharedRef();
+	return NewDockTab;
 }
 
 TSharedPtr<FWorkspaceItem> FObjectMixerEditorModule::GetWorkspaceGroup()
diff --git a/Engine/Plugins/Editor/ObjectMixer/ObjectMixer/Source/ObjectMixer/Public/ObjectMixerEditorModule.h b/Engine/Plugins/Editor/ObjectMixer/ObjectMixer/Source/ObjectMixer/Public/ObjectMixerEditorModule.h
--- a/Engine/Plugins/Editor/ObjectMixer/ObjectMixer/Source/ObjectMixer/Public/ObjectMixerEditorModule.h
+++ b/Engine/Plugins/Editor/ObjectMixer/ObjectMixer/Source/ObjectMixer/Public/ObjectMixerEditorModule.h
@@ -31,9 +31,18 @@ public:
 	
 	virtual TSharedPtr<SWidget> MakeObjectMixerDialog(
 		TSubclassOf<UObjectMixerObjectFilter> InDefaultFilterClass = nullptr);
+
+	/**
+	 * Tries to find the nomad tab assigned to this instance of Object Mixer.
+	 * If DockTab is not set, will try to find the tab using GetTabSpawnerId().
+	 */
+	virtual TSharedPtr<SDockTab> FindNomadTab();
 	
-	/** Build the List widget from scratch. */
-	void RegenerateListWidget();
+	/**
+	 * Build the List widget from scratch. If DockTab is not set, will try to find the tab using GetTabSpawnerId().
+	 * @return True if the widget was regenerated. False if the DockTab was invalid and could not be found.
+	 */
+	bool RegenerateListWidget();
 
 	/**
 	 * Regenerate the list items and refresh the list. Call when adding or removing variables.
@@ -90,7 +99,8 @@ protected:
 	/** The text that appears on the spawned nomad tab */
 	FText TabLabel;
 
-	TSharedPtr<SDockTab> DockTab;
+	/** The actual spawned nomad tab */
+	TWeakPtr<SDockTab> DockTab;
 
 	/** Menu Item variables */
 	FText MenuItemName;
