Temporary fix: Do not call cache and rebuild filters for blueprint recompile. Actual fix will be addressed later. This is to unblock Dev-Cine team.

diff --git a/Engine/Plugins/Editor/ObjectMixer/ObjectMixer/Source/ObjectMixer/Private/ObjectMixerEditorModule.cpp b/Engine/Plugins/Editor/ObjectMixer/ObjectMixer/Source/ObjectMixer/Private/ObjectMixerEditorModule.cpp
--- a/Engine/Plugins/Editor/ObjectMixer/ObjectMixer/Source/ObjectMixer/Private/ObjectMixerEditorModule.cpp
+++ b/Engine/Plugins/Editor/ObjectMixer/ObjectMixer/Source/ObjectMixer/Private/ObjectMixerEditorModule.cpp
@@ -336,10 +336,10 @@ void FObjectMixerEditorModule::BindDelegates()
 
 	DelegateHandles.Add(GEditor->OnBlueprintCompiled().AddLambda([this] ()
 	{
-		if (ListModel)
-		{
-			ListModel->CacheAndRebuildFilters(true);
-		}
+//		if (ListModel)
+//		{
+//			ListModel->CacheAndRebuildFilters(true);
+//		}
 	}));
 	DelegateHandles.Add(GEditor->OnLevelActorAdded().AddLambda([this] (AActor*)
 	{
