Object Mixer:

-Fix: Engine Hitch/Delay when opening Light Mixer

diff --git a/Engine/Plugins/Editor/ObjectMixer/ObjectMixer/Source/ObjectMixer/Private/Views/List/SObjectMixerEditorList.cpp b/Engine/Plugins/Editor/ObjectMixer/ObjectMixer/Source/ObjectMixer/Private/Views/List/SObjectMixerEditorList.cpp
--- a/Engine/Plugins/Editor/ObjectMixer/ObjectMixer/Source/ObjectMixer/Private/Views/List/SObjectMixerEditorList.cpp
+++ b/Engine/Plugins/Editor/ObjectMixer/ObjectMixer/Source/ObjectMixer/Private/Views/List/SObjectMixerEditorList.cpp
@@ -95,11 +95,43 @@ void SObjectMixerEditorList::Construct(const FArguments& InArgs, TSharedRef<FObj
 	RebuildCollectionSelector();
 
 	SetSingleCollectionSelection();
+
+	// Horizontal Scrollbar
+	TSharedPtr<SOverlay> ParentOverlay = StaticCastSharedPtr<SOverlay>(GetTreeView()->GetParentWidget());
+	FChildren* ScrollboxChildren = ParentOverlay->GetChildren();
+	
+	for (int32 ScrollboxChildItr = 0; ScrollboxChildItr < ScrollboxChildren->Num(); ScrollboxChildItr++)
+	{
+		TSharedRef<SWidget> ScrollboxChild = ScrollboxChildren->GetChildAt(ScrollboxChildItr);
+
+		if (ScrollboxChild == GetTreeView())
+		{
+			ParentOverlay->RemoveSlot(GetTreeView().ToSharedRef());
+			ParentOverlay->AddSlot(ScrollboxChildItr)
+			.HAlign(HAlign_Fill)
+			[
+				SNew(SScrollBox)
+				.Orientation(Orient_Horizontal)
+
+				+SScrollBox::Slot()
+				.HAlign(HAlign_Fill)
+				.VAlign(VAlign_Fill)
+				[
+					SAssignNew(ScrollBoxContainerPtr, SBox)
+					.HAlign(HAlign_Fill)
+					[
+						GetTreeView().ToSharedRef()
+					]
+				]
+			];
+		}
+	}
 }
 
 SObjectMixerEditorList::~SObjectMixerEditorList()
 {	
 	FLevelEditorSequencerIntegration::Get().GetOnSequencersChanged().RemoveAll(this);
+	HeaderRowContextMenuWidget.Reset();
 }
 
 FReply SObjectMixerEditorList::OnKeyDown(const FGeometry& MyGeometry, const FKeyEvent& InKeyEvent)
@@ -585,6 +617,12 @@ void SObjectMixerEditorList::SetTreeViewItemExpanded(const TSharedPtr<ISceneOutl
 	GetTreeView()->SetItemExpansion(RowToExpand, bNewExpansion);
 }
 
+void SObjectMixerEditorList::SetScrollBoxContainerWidthOverride(const float InScrollBoxContainerWidthOverride)
+{
+	ScrollBoxContainerWidthOverride = InScrollBoxContainerWidthOverride;
+	ScrollBoxContainerPtr->SetWidthOverride(GetScrollBoxContainerWidthOverride());
+}
+
 FObjectMixerSceneOutlinerColumnInfo* SObjectMixerEditorList::GetColumnInfoByPropertyName(const FName& InPropertyName)
 {
 	return Algo::FindByPredicate(HeaderColumnInfos,
@@ -663,6 +701,61 @@ void SObjectMixerEditorList::Tick(const FGeometry& AllottedGeometry, const doubl
 		RebuildList();
 	}
 
+	if (!bHasBeenResized)
+	{
+		TSharedPtr<SHeaderRow> HeaderRow = GetTree().GetHeaderRow();
+		constexpr float MinAllowedWidth = 40;
+		constexpr float MaxAllowedWidth = 100;
+		const TArray<FName> FixedColumns =
+			{
+				FSceneOutlinerBuiltInColumnTypes::Unsaved(),
+				FObjectMixerOutlinerVisibilityColumn::GetID(),
+				FObjectMixerOutlinerSoloColumn::GetID()
+			};
+
+		// Iterate over columns in reverse so that built-in columns don't get squished
+		const TIndirectArray<SHeaderRow::FColumn>& HeaderColumns = HeaderRow->GetColumns();
+		for (int32 ColumnItr = HeaderColumns.Num() - 1; ColumnItr >= 0; ColumnItr--)
+		{
+			const SHeaderRow::FColumn& Column = HeaderColumns[ColumnItr];
+			
+			if (Column.ColumnId == FSceneOutlinerBuiltInColumnTypes::Label())
+			{
+				constexpr float DefaultWidthForItemLabelColumn = 150;
+				HeaderRow->SetColumnWidth(Column.ColumnId, DefaultWidthForItemLabelColumn);
+			}
+			else if (Column.ColumnId == FSceneOutlinerBuiltInColumnTypes::ActorInfo())
+			{
+				constexpr float DefaultWidthForActorInfoColumn = 150;
+				HeaderRow->SetColumnWidth(Column.ColumnId, DefaultWidthForActorInfoColumn);
+			}
+			else if (FixedColumns.Contains(Column.ColumnId)) // Skip Fixed-width columns
+			{
+				continue;
+			}
+
+			const FVector2f ColumnDesiredSize = Column.HeaderContent.Widget->GetDesiredSize();
+			const FVector2f ColumnLocalSize = Column.HeaderContent.Widget->GetPaintSpaceGeometry().GetLocalSize();
+			float MaxCalculatedWidth = FMath::Max(ColumnDesiredSize.X, ColumnLocalSize.X);
+			MaxCalculatedWidth = FMath::Max(MinAllowedWidth, MaxCalculatedWidth);
+			MaxCalculatedWidth = FMath::Min(MaxAllowedWidth, MaxCalculatedWidth);
+			HeaderRow->SetColumnWidth(Column.ColumnId, MaxCalculatedWidth);
+
+			if (Column.bIsVisible)
+			{
+				SetScrollBoxContainerWidthOverride(GetScrollBoxContainerWidthOverride() + MaxCalculatedWidth);
+			}
+		}
+
+		const float ParentWidgetWidth = AllottedGeometry.GetLocalSize().X;
+		if (GetScrollBoxContainerWidthOverride() < ParentWidgetWidth)
+		{
+			SetScrollBoxContainerWidthOverride(ParentWidgetWidth);
+		}
+
+		bHasBeenResized = true;
+	}
+
 	if (PendingPropertyPropagations.Num() > 0)
 	{
 		PropagatePropertyChangesToSelectedRows();
@@ -675,77 +768,92 @@ void SObjectMixerEditorList::Tick(const FGeometry& AllottedGeometry, const doubl
 
 TSharedRef<SWidget> SObjectMixerEditorList::GenerateHeaderRowContextMenu()
 {
-	FMenuBuilder MenuBuilder(false, nullptr);
-	
-	MenuBuilder.AddSearchWidget();
-
-	MenuBuilder.BeginSection(TEXT("Actions"), LOCTEXT("ContextMenuActions", "Actions"));
+	if (!HeaderRowContextMenuWidget)
 	{
-		MenuBuilder.AddMenuEntry(
-			LOCTEXT("RestoreDefaultPropertyColumns", "Restore Default Columns"),
-			LOCTEXT("RestoreDefaultPropertyColumnsTooltip", "Restore property columns to built-in columns and those as defined in GetColumnsToShowByDefault() in your object filters."),
-			FSlateIcon(),
-			FUIAction(
-				FExecuteAction::CreateLambda([this]()
-				{
-					CloseAllMenus();
-					RestoreDefaultPropertyColumns();
-				})
-			)
-		);
-	}
-	MenuBuilder.EndSection();
+		FMenuBuilder MenuBuilder(false, nullptr);
+	
+		MenuBuilder.AddSearchWidget();
 
-	FName LastPropertyCategoryName = NAME_None;
+		MenuBuilder.BeginSection(TEXT("Actions"), LOCTEXT("ContextMenuActions", "Actions"));
+		{
+			MenuBuilder.AddMenuEntry(
+				LOCTEXT("RestoreDefaultPropertyColumns", "Restore Default Columns"),
+				LOCTEXT("RestoreDefaultPropertyColumnsTooltip", "Restore property columns to built-in columns and those as defined in GetColumnsToShowByDefault() in your object filters."),
+				FSlateIcon(),
+				FUIAction(
+					FExecuteAction::CreateLambda([this]()
+					{
+						CloseAllMenus();
+						RestoreDefaultPropertyColumns();
+					})
+				)
+			);
+		}
+		MenuBuilder.EndSection();
 
-	for (const FObjectMixerSceneOutlinerColumnInfo& ColumnInfo : HeaderColumnInfos)
-	{
-		const FName& PropertyCategoryName = ColumnInfo.PropertyCategoryName;
+		FName LastPropertyCategoryName = NAME_None;
 
-		if (!PropertyCategoryName.IsEqual(LastPropertyCategoryName))
+		for (const FObjectMixerSceneOutlinerColumnInfo& ColumnInfo : HeaderColumnInfos)
 		{
-			LastPropertyCategoryName = PropertyCategoryName;
+			const FName& PropertyCategoryName = ColumnInfo.PropertyCategoryName;
+
+			if (!PropertyCategoryName.IsEqual(LastPropertyCategoryName))
+			{
+				LastPropertyCategoryName = PropertyCategoryName;
 			
-			MenuBuilder.EndSection();
-            MenuBuilder.BeginSection(LastPropertyCategoryName, FText::FromName(LastPropertyCategoryName));
-		}
+				MenuBuilder.EndSection();
+				MenuBuilder.BeginSection(LastPropertyCategoryName, FText::FromName(LastPropertyCategoryName));
+			}
 		
-		const FName& ColumnID = ColumnInfo.ColumnID;
+			const FName& ColumnID = ColumnInfo.ColumnID;
 		
-		const FText Tooltip = ColumnInfo.PropertyRef ?
-			ColumnInfo.PropertyRef->GetToolTipText() : ColumnInfo.PropertyDisplayText;
+			const FText Tooltip = ColumnInfo.PropertyRef ?
+				ColumnInfo.PropertyRef->GetToolTipText() : ColumnInfo.PropertyDisplayText;
 
-		const bool bCanSelectColumn = ColumnInfo.bCanBeHidden;
+			const bool bCanSelectColumn = ColumnInfo.bCanBeHidden;
 
-		const FName Hook = ColumnInfo.PropertyType == EListViewColumnType::BuiltIn ? "Builtin" : "GeneratedProperties";
+			const FName Hook = ColumnInfo.PropertyType == EListViewColumnType::BuiltIn ? "Builtin" : "GeneratedProperties";
 		
-		MenuBuilder.AddMenuEntry(
-			ColumnInfo.PropertyDisplayText,
-			Tooltip,
-			FSlateIcon(),
-			FUIAction(
-				FExecuteAction::CreateLambda([this, ColumnID]()
-				{
-					TSharedPtr<FObjectMixerEditorList> PinnedListModel = ListModelPtr.Pin();
-					TSharedPtr<SHeaderRow> PinnedHeaderRow = GetTree().GetHeaderRow();
-					check(PinnedHeaderRow);
+			MenuBuilder.AddMenuEntry(
+				ColumnInfo.PropertyDisplayText,
+				Tooltip,
+				FSlateIcon(),
+				FUIAction(
+					FExecuteAction::CreateLambda([this, ColumnID]()
+					{
+						TSharedPtr<FObjectMixerEditorList> PinnedListModel = ListModelPtr.Pin();
+						TSharedPtr<SHeaderRow> PinnedHeaderRow = GetTree().GetHeaderRow();
+						check(PinnedHeaderRow);
 
-					const bool bNewColumnEnabled = !PinnedHeaderRow->IsColumnVisible(ColumnID);
+						const bool bNewColumnEnabled = !PinnedHeaderRow->IsColumnVisible(ColumnID);
 					
-					PinnedHeaderRow->SetShowGeneratedColumn(ColumnID, bNewColumnEnabled);
-				}),
-				FCanExecuteAction::CreateLambda([bCanSelectColumn](){return bCanSelectColumn;}),
-				FIsActionChecked::CreateLambda([this, ColumnID]()
-				{
-					return GetTree().GetHeaderRow()->IsColumnVisible(ColumnID);
-				})
-			),
-			Hook,
-			EUserInterfaceActionType::ToggleButton
-		);
+						PinnedHeaderRow->SetShowGeneratedColumn(ColumnID, bNewColumnEnabled);
+
+						if (const SHeaderRow::FColumn* Match = Algo::FindByPredicate(PinnedHeaderRow->GetColumns(),
+							[ColumnID](const SHeaderRow::FColumn& Column)
+							{
+								return Column.ColumnId == ColumnID;
+							}))
+						{
+							const float DesiredSizeOffset = Match->HeaderContent.Widget->GetDesiredSize().X * (bNewColumnEnabled ? 1 : -1);
+							SetScrollBoxContainerWidthOverride(GetScrollBoxContainerWidthOverride() + DesiredSizeOffset);
+						}
+					}),
+					FCanExecuteAction::CreateLambda([bCanSelectColumn](){return bCanSelectColumn;}),
+					FIsActionChecked::CreateLambda([this, ColumnID]()
+					{
+						return GetTree().GetHeaderRow()->IsColumnVisible(ColumnID);
+					})
+				),
+				Hook,
+				EUserInterfaceActionType::ToggleButton
+			);
+		}
+
+		HeaderRowContextMenuWidget = MenuBuilder.MakeWidget();
 	}
 
-	return MenuBuilder.MakeWidget();
+	return HeaderRowContextMenuWidget.ToSharedRef();
 }
 
 bool SObjectMixerEditorList::AddUniquePropertyColumnInfo(
diff --git a/Engine/Plugins/Editor/ObjectMixer/ObjectMixer/Source/ObjectMixer/Public/Views/List/SObjectMixerEditorList.h b/Engine/Plugins/Editor/ObjectMixer/ObjectMixer/Source/ObjectMixer/Public/Views/List/SObjectMixerEditorList.h
--- a/Engine/Plugins/Editor/ObjectMixer/ObjectMixer/Source/ObjectMixer/Public/Views/List/SObjectMixerEditorList.h
+++ b/Engine/Plugins/Editor/ObjectMixer/ObjectMixer/Source/ObjectMixer/Public/Views/List/SObjectMixerEditorList.h
@@ -155,6 +155,13 @@ public:
 	void SetTreeViewItemExpanded(const TSharedPtr<ISceneOutlinerTreeItem>& RowToExpand, const bool bNewExpansion) const;
 
 	void PropagatePropertyChangesToSelectedRows();
+	
+	[[nodiscard]] float GetScrollBoxContainerWidthOverride() const
+	{
+		return ScrollBoxContainerWidthOverride;
+	}
+
+	void SetScrollBoxContainerWidthOverride(const float InScrollBoxContainerWidthOverride);
 
 	// Columns
 
@@ -180,8 +187,10 @@ protected:
 	TSet<FObjectMixerEditorListRowData::FPropertyPropagationInfo> PendingPropertyPropagations;
 
 	TArray<FObjectMixerSceneOutlinerColumnInfo> HeaderColumnInfos;
+	TSharedPtr<SWidget> HeaderRowContextMenuWidget;
 
 	TWeakPtr<SHorizontalBox> ToolbarPtr;
+	TSharedPtr<SBox> ScrollBoxContainerPtr;
 
 	bool CanCreateFolder() const;
 	TSharedRef<SWidget> OnGenerateAddObjectButtonMenu() const;
@@ -218,4 +227,10 @@ protected:
 	void CreateActorTextInfoColumns(FSceneOutlinerInitializationOptions& OutInitOptions);
 	void SetupColumns(FSceneOutlinerInitializationOptions& OutInitOptions);
 	void SortColumnsForHeaderContextMenu();
+
+private:
+
+	bool bHasBeenResized = false;
+
+	float ScrollBoxContainerWidthOverride = 0.0f;
 };
