Object Mixer:

-Refactor Object Mixer refresh on Blueprint Compile

diff --git a/Engine/Plugins/Editor/ObjectMixer/ObjectMixer/Source/ObjectMixer/Private/ObjectFilter/ObjectMixerEditorObjectFilter.cpp b/Engine/Plugins/Editor/ObjectMixer/ObjectMixer/Source/ObjectMixer/Private/ObjectFilter/ObjectMixerEditorObjectFilter.cpp
--- a/Engine/Plugins/Editor/ObjectMixer/ObjectMixer/Source/ObjectMixer/Private/ObjectFilter/ObjectMixerEditorObjectFilter.cpp
+++ b/Engine/Plugins/Editor/ObjectMixer/ObjectMixer/Source/ObjectMixer/Private/ObjectFilter/ObjectMixerEditorObjectFilter.cpp
@@ -2,8 +2,21 @@
 
 #include "ObjectFilter/ObjectMixerEditorObjectFilter.h"
 
+#include "ObjectMixerEditorModule.h"
+
 #include "Kismet2/ComponentEditorUtils.h"
 
+void UObjectMixerObjectFilter::PostCDOCompiled(const FPostCDOCompiledContext& Context)
+{
+	UObject::PostCDOCompiled(Context);
+
+	// Only call on in-editor compile
+	if (!Context.bIsRegeneratingOnLoad)
+	{
+		FObjectMixerEditorModule::Get().OnBlueprintFilterCompiled().Broadcast();
+	}
+}
+
 FText UObjectMixerObjectFilter::GetRowDisplayName(UObject* InObject, const bool bIsHybridRow) const
 {
 	if (IsValid(InObject))
diff --git a/Engine/Plugins/Editor/ObjectMixer/ObjectMixer/Source/ObjectMixer/Private/ObjectMixerEditorModule.cpp b/Engine/Plugins/Editor/ObjectMixer/ObjectMixer/Source/ObjectMixer/Private/ObjectMixerEditorModule.cpp
--- a/Engine/Plugins/Editor/ObjectMixer/ObjectMixer/Source/ObjectMixer/Private/ObjectMixerEditorModule.cpp
+++ b/Engine/Plugins/Editor/ObjectMixer/ObjectMixer/Source/ObjectMixer/Private/ObjectMixerEditorModule.cpp
@@ -78,7 +78,6 @@ void FObjectMixerEditorModule::Teardown()
 
 	if (GEditor)
 	{
-		GEditor->OnBlueprintCompiled().RemoveAll(this);
 		GEditor->OnLevelActorAdded().RemoveAll(this);
 		GEditor->OnLevelActorDeleted().RemoveAll(this);
 		GEditor->GetSelectedActors()->SelectionChangedEvent.RemoveAll(this);
@@ -334,13 +333,6 @@ void FObjectMixerEditorModule::BindDelegates()
 		RequestRebuildList();
 	}));
 
-	DelegateHandles.Add(GEditor->OnBlueprintCompiled().AddLambda([this] ()
-	{
-//		if (ListModel)
-//		{
-//			ListModel->CacheAndRebuildFilters(true);
-//		}
-	}));
 	DelegateHandles.Add(GEditor->OnLevelActorAdded().AddLambda([this] (AActor*)
 	{
 		RequestRebuildList();
diff --git a/Engine/Plugins/Editor/ObjectMixer/ObjectMixer/Source/ObjectMixer/Private/Views/List/ObjectMixerEditorList.cpp b/Engine/Plugins/Editor/ObjectMixer/ObjectMixer/Source/ObjectMixer/Private/Views/List/ObjectMixerEditorList.cpp
--- a/Engine/Plugins/Editor/ObjectMixer/ObjectMixer/Source/ObjectMixer/Private/Views/List/ObjectMixerEditorList.cpp
+++ b/Engine/Plugins/Editor/ObjectMixer/ObjectMixer/Source/ObjectMixer/Private/Views/List/ObjectMixerEditorList.cpp
@@ -14,6 +14,13 @@
 void FObjectMixerEditorList::Initialize()
 {
 	RegisterAndMapContextMenuCommands();
+
+	// The base module is always notified when a blueprint filter is compiled.
+	// When it is, all submodules should rebuild their lists.
+	OnBlueprintFilterCompiledHandle = FObjectMixerEditorModule::Get().OnBlueprintFilterCompiled().AddLambda([this]()
+	{
+		CacheAndRebuildFilters(true);
+	});
 }
 
 void FObjectMixerEditorList::RegisterAndMapContextMenuCommands()
@@ -327,6 +334,11 @@ FObjectMixerEditorList::FObjectMixerEditorList(const FName InModuleName)
 FObjectMixerEditorList::~FObjectMixerEditorList()
 {
 	FlushWidget();
+
+	if (FObjectMixerEditorModule* ObjectMixerEditorModule = FModuleManager::GetModulePtr<FObjectMixerEditorModule>("ObjectMixer"))
+	{
+		ObjectMixerEditorModule->OnBlueprintFilterCompiled().Remove(OnBlueprintFilterCompiledHandle);
+	}
 }
 
 void FObjectMixerEditorList::FlushWidget()
diff --git a/Engine/Plugins/Editor/ObjectMixer/ObjectMixer/Source/ObjectMixer/Public/ObjectFilter/ObjectMixerEditorObjectFilter.h b/Engine/Plugins/Editor/ObjectMixer/ObjectMixer/Source/ObjectMixer/Public/ObjectFilter/ObjectMixerEditorObjectFilter.h
--- a/Engine/Plugins/Editor/ObjectMixer/ObjectMixer/Source/ObjectMixer/Public/ObjectFilter/ObjectMixerEditorObjectFilter.h
+++ b/Engine/Plugins/Editor/ObjectMixer/ObjectMixer/Source/ObjectMixer/Public/ObjectFilter/ObjectMixerEditorObjectFilter.h
@@ -51,6 +51,10 @@ class OBJECTMIXEREDITOR_API UObjectMixerObjectFilter : public UObject
 public:
 
 	UObjectMixerObjectFilter() = default;
+
+	/** Begin UObject overrides */
+	virtual void PostCDOCompiled(const FPostCDOCompiledContext& Context) override;
+	/** End UObject overrides */
 	
 	/**
 	 * Return the basic object types you want to filter for in your level.
diff --git a/Engine/Plugins/Editor/ObjectMixer/ObjectMixer/Source/ObjectMixer/Public/ObjectMixerEditorModule.h b/Engine/Plugins/Editor/ObjectMixer/ObjectMixer/Source/ObjectMixer/Public/ObjectMixerEditorModule.h
--- a/Engine/Plugins/Editor/ObjectMixer/ObjectMixer/Source/ObjectMixer/Public/ObjectMixerEditorModule.h
+++ b/Engine/Plugins/Editor/ObjectMixer/ObjectMixer/Source/ObjectMixer/Public/ObjectMixerEditorModule.h
@@ -85,6 +85,12 @@ public:
 	const TSubclassOf<UObjectMixerObjectFilter>& GetDefaultFilterClass() const;
 
 	const static FName BaseObjectMixerModuleName;
+	
+	DECLARE_MULTICAST_DELEGATE(FOnBlueprintFilterCompiled);
+	FOnBlueprintFilterCompiled& OnBlueprintFilterCompiled()
+	{
+		return OnBlueprintFilterCompiledDelegate;
+	}
 
 protected:
 
@@ -116,6 +122,8 @@ protected:
 	
 	TSet<FDelegateHandle> DelegateHandles;
 
+	FOnBlueprintFilterCompiled OnBlueprintFilterCompiledDelegate;
+
 private:
 
 	TSharedPtr<FWorkspaceItem> WorkspaceGroup;
diff --git a/Engine/Plugins/Editor/ObjectMixer/ObjectMixer/Source/ObjectMixer/Public/Views/List/ObjectMixerEditorList.h b/Engine/Plugins/Editor/ObjectMixer/ObjectMixer/Source/ObjectMixer/Public/Views/List/ObjectMixerEditorList.h
--- a/Engine/Plugins/Editor/ObjectMixer/ObjectMixer/Source/ObjectMixer/Public/Views/List/ObjectMixerEditorList.h
+++ b/Engine/Plugins/Editor/ObjectMixer/ObjectMixer/Source/ObjectMixer/Public/Views/List/ObjectMixerEditorList.h
@@ -224,4 +224,6 @@ protected:
 	EObjectMixerTreeViewMode TreeViewMode = EObjectMixerTreeViewMode::Folders;
 
 	FName ModuleName = NAME_None;
+	
+	FDelegateHandle OnBlueprintFilterCompiledHandle;
 };
