Object Mixer:

-Audit filter class override functions

diff --git a/Engine/Plugins/Editor/ObjectMixer/ObjectMixer/Source/ObjectMixer/Private/ObjectFilter/ObjectMixerEditorObjectFilter.cpp b/Engine/Plugins/Editor/ObjectMixer/ObjectMixer/Source/ObjectMixer/Private/ObjectFilter/ObjectMixerEditorObjectFilter.cpp
--- a/Engine/Plugins/Editor/ObjectMixer/ObjectMixer/Source/ObjectMixer/Private/ObjectFilter/ObjectMixerEditorObjectFilter.cpp
+++ b/Engine/Plugins/Editor/ObjectMixer/ObjectMixer/Source/ObjectMixer/Private/ObjectFilter/ObjectMixerEditorObjectFilter.cpp
@@ -4,8 +4,6 @@
 
 #include "ObjectMixerEditorModule.h"
 
-#include "Kismet2/ComponentEditorUtils.h"
-
 void UObjectMixerObjectFilter::PostCDOCompiled(const FPostCDOCompiledContext& Context)
 {
 	UObject::PostCDOCompiled(Context);
@@ -17,116 +15,11 @@ void UObjectMixerObjectFilter::PostCDOCompiled(const FPostCDOCompiledContext& Co
 	}
 }
 
-FText UObjectMixerObjectFilter::GetRowDisplayName(UObject* InObject, const bool bIsHybridRow) const
-{
-	if (IsValid(InObject))
-	{
-		if (const AActor* AsActor = Cast<AActor>(InObject))
-		{
-			return FText::FromString(AsActor->GetActorLabel());
-		}
-
-		if (bIsHybridRow)
-		{
-			if (const AActor* Outer = InObject->GetTypedOuter<AActor>())
-			{
-				return FText::FromString(Outer->GetActorLabel());
-			}
-		}
-		
-		if (const UActorComponent* AsActorComponent = Cast<UActorComponent>(InObject))
-		{
-			const FName ComponentName = FComponentEditorUtils::FindVariableNameGivenComponentInstance(AsActorComponent);
-			if (ComponentName != NAME_None)
-			{
-				return FText::FromName(ComponentName);
-			}
-		}
-		
-		return FText::FromString(InObject->GetName());
-	}
-
-	return FText::GetEmpty();
-}
-
-FText UObjectMixerObjectFilter::GetRowTooltipText(UObject* InObject, const bool bIsHybridRow) const
-{
-	if (IsValid(InObject))
-	{
-		if (const AActor* Outer = InObject->GetTypedOuter<AActor>())
-		{
-			return FText::Format(INVTEXT("{0} ({1})"), FText::FromString(Outer->GetActorLabel()), FText::FromString(InObject->GetName()));
-		}
-		
-		if (const AActor* AsActor = Cast<AActor>(InObject))
-		{
-			return FText::FromString(AsActor->GetActorLabel());
-		}
-		
-		if (const UActorComponent* AsActorComponent = Cast<UActorComponent>(InObject))
-		{
-			const FName ComponentName = FComponentEditorUtils::FindVariableNameGivenComponentInstance(AsActorComponent);
-			if (ComponentName != NAME_None)
-			{
-				if (!ComponentName.IsEqual(AsActorComponent->GetFName()))
-				{
-					return
-					   FText::Format(INVTEXT("{0} ({1})"),
-						   FText::FromName(ComponentName),
-						   FText::FromName(AsActorComponent->GetFName())
-					   );
-				}
-
-				return FText::FromName(ComponentName);
-			}
-		}
-		
-		return FText::FromString(InObject->GetName());
-	}
-
-	return FText::GetEmpty();
-}
-
 bool UObjectMixerObjectFilter::GetShowTransientObjects() const
 {
 	return false;
 }
 
-bool UObjectMixerObjectFilter::GetRowEditorVisibility(UObject* InObject) const
-{
-	if (IsValid(InObject))
-	{
-		TObjectPtr<AActor> Actor = Cast<AActor>(InObject);
-	
-		if (!Actor)
-		{
-			Actor = InObject->GetTypedOuter<AActor>();
-		}
-
-		return Actor ? !Actor->IsTemporarilyHiddenInEditor() : false;
-	}
-
-	return false;
-}
-
-void UObjectMixerObjectFilter::OnSetRowEditorVisibility(UObject* InObject, bool bNewIsVisible) const
-{
-	if (IsValid(InObject))
-	{
-		TObjectPtr<AActor> Actor = Cast<AActor>(InObject);
-	
-		if (!Actor)
-		{
-			Actor = InObject->GetTypedOuter<AActor>();
-		}
-
-		if (Actor)
-		{
-			Actor->SetIsTemporarilyHiddenInEditor(!bNewIsVisible);
-		}
-	}
-}
-
 TSet<FName> UObjectMixerObjectFilter::GetColumnsToShowByDefault() const
 {
 	return {};
@@ -239,31 +132,3 @@ TSet<UClass*> UObjectMixerObjectFilter::GetParentAndChildClassesFromSpecifiedCla
 
 	return GetParentAndChildClassesFromSpecifiedClasses(AsUClasses, Options);
 }
-
-EFieldIterationFlags UObjectMixerObjectFilter::GetDesiredFieldIterationFlags(const bool bIncludeInheritedProperties)
-{
-	return bIncludeInheritedProperties ? EFieldIterationFlags::IncludeSuper : EFieldIterationFlags::Default;
-}
-
-TSet<FName> UObjectMixerObjectFilter::GenerateIncludeListFromExcludeList(const TSet<FName>& ExcludeList) const
-{
-	TSet<FName> IncludeList;
-
-	const EObjectMixerInheritanceInclusionOptions Options =
-		GetObjectMixerPropertyInheritanceInclusionOptions();
-	TSet<UClass*> SpecifiedClasses =
-		GetParentAndChildClassesFromSpecifiedClasses(GetObjectClassesToFilter(), Options);
-	
-	for (const UClass* Class : SpecifiedClasses)
-	{
-		for (TFieldIterator<FProperty> FieldIterator(Class); FieldIterator; ++FieldIterator)
-		{
-			if (const FProperty* Property = *FieldIterator)
-			{
-				IncludeList.Add(Property->GetFName());
-			}
-		}
-	}
-
-	return IncludeList.Difference(ExcludeList);
-}
diff --git a/Engine/Plugins/Editor/ObjectMixer/ObjectMixer/Source/ObjectMixer/Private/Views/List/ObjectMixerEditorListRowData.cpp b/Engine/Plugins/Editor/ObjectMixer/ObjectMixer/Source/ObjectMixer/Private/Views/List/ObjectMixerEditorListRowData.cpp
--- a/Engine/Plugins/Editor/ObjectMixer/ObjectMixer/Source/ObjectMixer/Private/Views/List/ObjectMixerEditorListRowData.cpp
+++ b/Engine/Plugins/Editor/ObjectMixer/ObjectMixer/Source/ObjectMixer/Private/Views/List/ObjectMixerEditorListRowData.cpp
@@ -91,25 +91,6 @@ bool FObjectMixerEditorListRowData::HasAtLeastOneChildThatIsNotSolo(const TShare
 	return false;
 }
 
-FText FObjectMixerEditorListRowData::GetDisplayName(TSharedPtr<ISceneOutlinerTreeItem> InTreeItem) const
-{
-	const FText Override = GetDisplayNameOverride();
-	if (!Override.IsEmpty())
-	{
-		return Override;
-	}
-
-	if (const UObjectMixerObjectFilter* Filter = GetMainObjectFilterInstance())
-	{
-		if (const TObjectPtr<UObject> Object = FObjectMixerUtils::GetRowObject(InTreeItem))
-		{
-			return Filter->GetRowDisplayName(Object, GetIsHybridRow());
-		}
-	}
-
-	return FText::GetEmpty();
-}
-
 TWeakPtr<SObjectMixerEditorList> FObjectMixerEditorListRowData::GetListView() const
 {
 	if (SceneOutlinerPtr.IsValid())
diff --git a/Engine/Plugins/Editor/ObjectMixer/ObjectMixer/Source/ObjectMixer/Public/ObjectFilter/ObjectMixerEditorObjectFilter.h b/Engine/Plugins/Editor/ObjectMixer/ObjectMixer/Source/ObjectMixer/Public/ObjectFilter/ObjectMixerEditorObjectFilter.h
--- a/Engine/Plugins/Editor/ObjectMixer/ObjectMixer/Source/ObjectMixer/Public/ObjectFilter/ObjectMixerEditorObjectFilter.h
+++ b/Engine/Plugins/Editor/ObjectMixer/ObjectMixer/Source/ObjectMixer/Public/ObjectFilter/ObjectMixerEditorObjectFilter.h
@@ -2,7 +2,6 @@
 
 #pragma once
 
-#include "Components/LightComponent.h"
 #include "UObject/Class.h"
 
 #include "ObjectMixerEditorObjectFilter.generated.h"
@@ -71,39 +70,11 @@ public:
 	 */
 	virtual TSet<TSubclassOf<AActor>> GetObjectClassesToPlace() const { return {}; }
 
-	/**
-	 * Get the text to display for the object name/label.
-	 * This is useful if one of your classes is a component type and you want the label of the component's owning actor, for example.
-	 * If not overridden, this returns the actor's label if it's an actor and the object's name for all other objects.
-	 * @param bIsHybridRow If true, this row represents a single matching component condensed into a single row with its actor parent.
-	 */
-	virtual FText GetRowDisplayName(UObject* InObject, const bool bIsHybridRow = false) const;
-
-	/**
-	 * Get the text to display for the row's tooltip when hovering over it.
-	 * @param bIsHybridRow If true, this row represents a single matching component condensed into a single row with its actor parent. (unused)
-	 */
-	virtual FText GetRowTooltipText(UObject* InObject, const bool bIsHybridRow = false) const;
-
 	/**
 	 * Determines if transient objects (such as Sequencer Spawnables) should be shown in the list. False by default.
 	 */
 	virtual bool GetShowTransientObjects() const;
 
-	/**
-	 * Controls how to display the row's visibility icon. Return true if the object should be visible.
-	 * Generally this should work like the Scene Outliner does.
-	 * If not overridden, we use the Editor Visibility of the object's AActor outer (unless it's an actor itself).
-	 */
-	virtual bool GetRowEditorVisibility(UObject* InObject) const;
-
-	/**
-	 * Controls what happens when the row's visibility icon is clicked.
-	 * Generally this should work like the Scene Outliner does.
-	 * If not overridden, we set the Editor Visibility of the object's AActor outer (unless it's an actor itself).
-	 */
-	virtual void OnSetRowEditorVisibility(UObject* InObject, bool bNewIsVisible) const;
-
 	/**
 	 * Specify a list of property names corresponding to columns you want to show by default.
 	 * For example, you can specify "Intensity" and "LightColor" to show only those property columns by default in the UI.
@@ -160,16 +131,6 @@ public:
 
 	static TSet<UClass*> GetParentAndChildClassesFromSpecifiedClasses(
 		const TSet<TSubclassOf<AActor>>& InSpecifiedClasses, EObjectMixerInheritanceInclusionOptions Options);
-
-	static EFieldIterationFlags GetDesiredFieldIterationFlags(const bool bIncludeInheritedProperties);
-
-protected:
-
-	/**
-	 * Given a set of property names you wish to include, returns a list of all other properties on InObject not found in ExcludeList.
-	 * Useful when defining default visible columns in a list view.
-	 */
-	TSet<FName> GenerateIncludeListFromExcludeList(const TSet<FName>& ExcludeList) const;
 };
 
 /**
@@ -209,32 +170,6 @@ public:
 		return Super::GetObjectClassesToPlace();
 	}
 
-	/**
-	 * Get the text to display for the object name/label.
-	 * This is useful if one of your classes is a component type and you want the label of the component's owning actor, for example.
-	 * If not overridden, this returns the actor's label if it's an actor and the object's name for all other objects.
-	 * @param bIsHybridRow If true, this row represents a single matching component condensed into a single row with its actor parent.
-	 */
-	UFUNCTION(BlueprintNativeEvent, BlueprintCallable, Category = "Object Mixer")
-	FText GetRowDisplayName(UObject* InObject, const bool bIsHybridRow) const override;
-	
-	FText GetRowDisplayName_Implementation(UObject* InObject, const bool bIsHybridRow) const
-	{
-		return Super::GetRowDisplayName(InObject, bIsHybridRow);
-	}
-
-	/**
-	 * Get the text to display for the row's tooltip when hovering over it.
-	 * @param bIsHybridRow If true, this row represents a single matching component condensed into a single row with its actor parent. (unused)
-	 */
-	UFUNCTION(BlueprintNativeEvent, BlueprintCallable, Category = "Object Mixer")
-	FText GetRowTooltipText(UObject* InObject, const bool bIsHybridRow) const override;
-	
-	FText GetRowTooltipText_Implementation(UObject* InObject, const bool bIsHybridRow) const
-	{
-		return Super::GetRowTooltipText(InObject, bIsHybridRow);
-	}
-
 	/**
 	 * Determines if transient objects (such as Sequencer Spawnables) should be shown in the list. False by default.
 	 */
@@ -246,32 +181,6 @@ public:
 		return Super::GetShowTransientObjects();
 	}
 	
-	/**
-	 * Controls how to display the row's visibility icon. Return true if the object should be visible.
-	 * Generally this should work like the Scene Outliner does.
-	 * If not overridden, we use the Editor Visibility of the object's AActor outer (unless it's an actor itself).
-	 */
-	UFUNCTION(BlueprintNativeEvent, BlueprintCallable, Category = "Object Mixer")
-	bool GetRowEditorVisibility(UObject* InObject) const override;
-
-	bool GetRowEditorVisibility_Implementation(UObject* InObject) const
-	{
-		return Super::GetRowEditorVisibility(InObject);
-	}
-	
-	/**
-	 * Controls what happens when the row's visibility icon is clicked.
-	 * Generally this should work like the Scene Outliner does.
-	 * If not overridden, we set the Editor Visibility of the object's AActor outer (unless it's an actor itself).
-	 */
-	UFUNCTION(BlueprintNativeEvent, BlueprintCallable, Category = "Object Mixer")
-	void OnSetRowEditorVisibility(UObject* InObject, bool bNewIsVisible) const override;
-
-	void OnSetRowEditorVisibility_Implementation(UObject* InObject, bool bNewIsVisible) const
-	{
-		return Super::OnSetRowEditorVisibility(InObject, bNewIsVisible);
-	}
-	
 	/**
 	 * Specify a list of property names corresponding to columns you want to show by default.
 	 * For example, you can specify "Intensity" and "LightColor" to show only those property columns by default in the UI.
diff --git a/Engine/Plugins/Editor/ObjectMixer/ObjectMixer/Source/ObjectMixer/Public/Views/List/ObjectMixerEditorListRowData.h b/Engine/Plugins/Editor/ObjectMixer/ObjectMixer/Source/ObjectMixer/Public/Views/List/ObjectMixerEditorListRowData.h
--- a/Engine/Plugins/Editor/ObjectMixer/ObjectMixer/Source/ObjectMixer/Public/Views/List/ObjectMixerEditorListRowData.h
+++ b/Engine/Plugins/Editor/ObjectMixer/ObjectMixer/Source/ObjectMixer/Public/Views/List/ObjectMixerEditorListRowData.h
@@ -80,8 +80,6 @@ struct OBJECTMIXEREDITOR_API FObjectMixerEditorListRowData
 	[[nodiscard]] bool HasAtLeastOneChildThatIsNotSolo(
 		const TSharedRef<ISceneOutlinerTreeItem> InRow, const bool bRecursive = true) const;
 
-	[[nodiscard]] FText GetDisplayName(TSharedPtr<ISceneOutlinerTreeItem> InTreeItem) const;
-
 	[[nodiscard]] const FText& GetDisplayNameOverride() const
 	{
 		return DisplayNameOverride;
