Object Mixer

-Fix: Removing a Light from a Light Collection does Not Immediately Update UI

diff --git a/Engine/Plugins/Editor/ObjectMixer/ObjectMixer/Source/ObjectMixer/Private/Views/List/ObjectMixerEditorList.cpp b/Engine/Plugins/Editor/ObjectMixer/ObjectMixer/Source/ObjectMixer/Private/Views/List/ObjectMixerEditorList.cpp
--- a/Engine/Plugins/Editor/ObjectMixer/ObjectMixer/Source/ObjectMixer/Private/Views/List/ObjectMixerEditorList.cpp
+++ b/Engine/Plugins/Editor/ObjectMixer/ObjectMixer/Source/ObjectMixer/Private/Views/List/ObjectMixerEditorList.cpp
@@ -162,7 +162,15 @@ bool FObjectMixerEditorList::RequestAddObjectsToCollection(const FName& Collecti
 		for (const TSubclassOf<UObjectMixerObjectFilter>& Class : GetObjectFilterClasses())
 		{
 			const FName FilterName = Class->GetFName();
-			return SerializedData->AddObjectsToCollection(FilterName, CollectionName, ObjectsToAdd);
+			if (SerializedData->AddObjectsToCollection(FilterName, CollectionName, ObjectsToAdd))
+			{
+				if (IsCollectionSelected(CollectionName))
+				{
+					RequestRebuildList();
+				}
+				return true;
+			}
+			return false;
 		}
 	}
 
@@ -178,7 +186,15 @@ bool FObjectMixerEditorList::RequestRemoveObjectsFromCollection(const FName& Col
 			const FName FilterName = Class->GetFName();
 			if (SerializedData->DoesCollectionExist(FilterName, CollectionName))
 			{
-				return SerializedData->RemoveObjectsFromCollection(FilterName, CollectionName, ObjectsToRemove);
+				if (SerializedData->RemoveObjectsFromCollection(FilterName, CollectionName, ObjectsToRemove))
+				{
+					if (IsCollectionSelected(CollectionName))
+					{
+						RequestRebuildList();
+					}
+					return true;
+				}
+				return false;
 			}
 		}
 	}
@@ -544,7 +560,7 @@ const TSet<FName>& FObjectMixerEditorList::GetSelectedCollections() const
 	return SelectedCollections;
 }
 
-bool FObjectMixerEditorList::IsCollectionSelected(const FName& CollectionName)
+bool FObjectMixerEditorList::IsCollectionSelected(const FName& CollectionName) const
 {
 	return GetSelectedCollections().Contains(CollectionName);
 }
diff --git a/Engine/Plugins/Editor/ObjectMixer/ObjectMixer/Source/ObjectMixer/Public/Views/List/ObjectMixerEditorList.h b/Engine/Plugins/Editor/ObjectMixer/ObjectMixer/Source/ObjectMixer/Public/Views/List/ObjectMixerEditorList.h
--- a/Engine/Plugins/Editor/ObjectMixer/ObjectMixer/Source/ObjectMixer/Public/Views/List/ObjectMixerEditorList.h
+++ b/Engine/Plugins/Editor/ObjectMixer/ObjectMixer/Source/ObjectMixer/Public/Views/List/ObjectMixerEditorList.h
@@ -175,7 +175,7 @@ public:
 	
 	/** Represents the collections the user has selected in the UI. If empty, "All" is considered as selected. */
 	UE_NODISCARD const TSet<FName>& GetSelectedCollections() const;
-	UE_NODISCARD bool IsCollectionSelected(const FName& CollectionName);
+	UE_NODISCARD bool IsCollectionSelected(const FName& CollectionName) const;
 	void SetSelectedCollections(const TSet<FName> InSelectedCollections);
 	void SetCollectionSelected(const FName& CollectionName, const bool bNewSelected);
 
