WorldPartition: Add UActorDescContainerInstance/FWorldPartitionActorDescInstance which take transient/context data from UActorDescContainer/FWorldPartitionActorDesc leaving them as pure on disk representations.
- UActorDescContainerInstance contains: Transform, ContainerID, InstancingContext
- FWorldPartitionActorDescInstance contains: ResolvedDataLayers, RefCounts, Remapped Instanced actor path, Actor loading code

#tests editor, cook, pie, hlod builder, level instance, internal projects

diff --git a/Engine/Plugins/Editor/DataValidation/Source/DataValidation/Private/WorldPartitionChangelistValidator.cpp b/Engine/Plugins/Editor/DataValidation/Source/DataValidation/Private/WorldPartitionChangelistValidator.cpp
--- a/Engine/Plugins/Editor/DataValidation/Source/DataValidation/Private/WorldPartitionChangelistValidator.cpp
+++ b/Engine/Plugins/Editor/DataValidation/Source/DataValidation/Private/WorldPartitionChangelistValidator.cpp
@@ -17,7 +17,6 @@
 #include "WorldPartition/DataLayer/DataLayerInstanceWithAsset.h"
 #include "WorldPartition/DataLayer/WorldDataLayers.h"
 #include "WorldPartition/ContentBundle/ContentBundlePaths.h"
-#include "WorldPartition/WorldPartitionActorDescView.h"
 
 #include UE_INLINE_GENERATED_CPP_BY_NAME(WorldPartitionChangelistValidator)
 
@@ -51,7 +50,7 @@ EDataValidationResult UWorldPartitionChangelistValidator::ValidateLoadedAsset_Im
 }
 
 // Extract all Actors/Map from Changelist (in OFPA this should be one Actor per Package, and we'll discard all Actors from non WorldPartition maps)
-// and add them to a Map of World->Files[] so that we can do one validation per world. Once Worlds are identified, we either the UActorDescContainer 
+// and add them to a Map of World->Files[] so that we can do one validation per world. Once Worlds are identified, we either the UActorDescContainerInstance 
 // from memory (if loaded) or request it to be loaded, we then build a Set of objects that interest us from the Actors in the CL 
 void UWorldPartitionChangelistValidator::ValidateActorsAndDataLayersFromChangeList(UDataValidationChangelist* Changelist)
 {
@@ -136,34 +135,34 @@ void UWorldPartitionChangelistValidator::ValidateActorsAndDataLayersFromChangeLi
 		}
 	}
 
-	auto RegisterContainerToValidate = [](UWorld* InWorld, FName InContainerPackageName, FActorDescContainerCollection& OutRegisteredContainers, const FGuid& InContentBundleGuid)
+	auto RegisterContainerToValidate = [](UWorld* InWorld, FName InContainerPackageName, FActorDescContainerInstanceCollection& OutRegisteredContainers, const FGuid& InContentBundleGuid)
 	{
 		if (OutRegisteredContainers.Contains(InContainerPackageName))
 		{
 			return;
 		}
 
-		UActorDescContainer* ActorDescContainer = nullptr;
+		UActorDescContainerInstance* ContainerInstance = nullptr;
 		if (InWorld != nullptr)
 		{
 			// World is Loaded reuse the ActorDescContainer of the Content Bundle
-			ActorDescContainer = InWorld->GetWorldPartition()->FindContainer(InContainerPackageName);
+			ContainerInstance = InWorld->GetWorldPartition()->FindContainer(InContainerPackageName);
 		}
 
 		// Even if world is valid, its world partition is not necessarily initialized
-		if (!ActorDescContainer)
+		if (!ContainerInstance)
 		{
-			// Find in memory failed, load the ActorDescContainer
-			ActorDescContainer = NewObject<UActorDescContainer>();
-			ActorDescContainer->Initialize({ nullptr, InContainerPackageName });
-			ActorDescContainer->SetContentBundleGuid(InContentBundleGuid);
+			// Find in memory failed, load the ActorDescContainerInstance
+			ContainerInstance = NewObject<UActorDescContainerInstance>();
+			ContainerInstance->Initialize({ InContainerPackageName });
+			ContainerInstance->GetContainer()->SetContentBundleGuid(InContentBundleGuid);
 		}
 		else
 		{
-			check(ActorDescContainer->GetContentBundleGuid() == InContentBundleGuid);
+			check(ContainerInstance->GetContentBundleGuid() == InContentBundleGuid);
 		}
 
-		OutRegisteredContainers.AddContainer(ActorDescContainer);
+		OutRegisteredContainers.AddContainer(ContainerInstance);
 	};
 
 	// For Each world 
@@ -177,7 +176,7 @@ void UWorldPartitionChangelistValidator::ValidateActorsAndDataLayersFromChangeLi
 		
 		TGuardValue<UObject*> GuardCurrentAsset(CurrentAsset, World);
 		
-		FActorDescContainerCollection ContainersToValidate;
+		FActorDescContainerInstanceCollection ContainersToValidate;
 
 		// Always register the main world container because content bundle containers can't be validated separately
 		RegisterContainerToValidate(World, MapPath.GetPackageName(), ContainersToValidate, FGuid());
@@ -204,9 +203,9 @@ void UWorldPartitionChangelistValidator::ValidateActorsAndDataLayersFromChangeLi
 		for (const FAssetData& ActorData : ActorsData)
 		{
 			// Get the actor descriptor
-			if (const FWorldPartitionActorDesc* ActorDesc = ContainersToValidate.GetActorDescByPath(ActorData.AssetName.ToString()))
+			if (const FWorldPartitionActorDescInstance* ActorDescInstance = ContainersToValidate.GetActorDescInstanceByPath(ActorData.AssetName.ToString()))
 			{
-				RelevantActorGuids.Add(ActorDesc->GetGuid());
+				RelevantActorGuids.Add(ActorDescInstance->GetGuid());
 			}
 		}
 
@@ -215,21 +214,21 @@ void UWorldPartitionChangelistValidator::ValidateActorsAndDataLayersFromChangeLi
 			.SetErrorHandler(this)
 			.SetEnableStreaming(!ULevel::GetIsStreamingDisabledFromPackage(MapPath.GetPackageName()));
 
-		ContainersToValidate.ForEachActorDescContainer([&Params](const UActorDescContainer* ActorDescContainer)
+		ContainersToValidate.ForEachActorDescContainer([&Params](const UActorDescContainerInstance* ContainerInstance)
 		{
-			for (FActorDescList::TConstIterator<> ActorDescIt(ActorDescContainer); ActorDescIt; ++ActorDescIt)
+			for (UActorDescContainerInstance::TConstIterator<> Iterator(ContainerInstance); Iterator; ++Iterator)
 			{
-				check(!Params.ActorGuidsToContainerMap.Contains(ActorDescIt->GetGuid()));
-				Params.ActorGuidsToContainerMap.Add(ActorDescIt->GetGuid(), ActorDescContainer);
+				check(!Params.ActorGuidsToContainerInstanceMap.Contains(Iterator->GetGuid()));
+				Params.ActorGuidsToContainerInstanceMap.Add(Iterator->GetGuid(), ContainerInstance);
 			}
 		});
 
-		Params.ActorDescContainerCollection = &ContainersToValidate;
+		Params.ActorDescContainerInstanceCollection = &ContainersToValidate;
 		UWorldPartition::CheckForErrors(Params);
 	}
 }
 
-bool UWorldPartitionChangelistValidator::Filter(const FWorldPartitionActorDescView& ActorDescView)
+bool UWorldPartitionChangelistValidator::Filter(const IWorldPartitionActorDescInstanceView& ActorDescView)
 {
 	if (RelevantActorGuids.Find(ActorDescView.GetGuid()))
 	{
@@ -251,7 +250,7 @@ bool UWorldPartitionChangelistValidator::Filter(const UDataLayerInstance* InData
 	return DataLayerWithAsset != nullptr && DataLayerWithAsset->GetAsset() != nullptr && RelevantDataLayerAssets.Contains(DataLayerWithAsset->GetAsset()->GetPathName());
 }
 
-void UWorldPartitionChangelistValidator::OnInvalidRuntimeGrid(const FWorldPartitionActorDescView& ActorDescView, FName GridName)
+void UWorldPartitionChangelistValidator::OnInvalidRuntimeGrid(const IWorldPartitionActorDescInstanceView& ActorDescView, FName GridName)
 {
 	if (Filter(ActorDescView))
 	{
@@ -263,7 +262,7 @@ void UWorldPartitionChangelistValidator::OnInvalidRuntimeGrid(const FWorldPartit
 	}
 }
 
-void UWorldPartitionChangelistValidator::OnInvalidReference(const FWorldPartitionActorDescView& ActorDescView, const FGuid& ReferenceGuid, FWorldPartitionActorDescView* ReferenceActorDescView)
+void UWorldPartitionChangelistValidator::OnInvalidReference(const IWorldPartitionActorDescInstanceView& ActorDescView, const FGuid& ReferenceGuid, IWorldPartitionActorDescInstanceView* ReferenceActorDescView)
 {
 	if (Filter(ActorDescView))
 	{
@@ -275,7 +274,7 @@ void UWorldPartitionChangelistValidator::OnInvalidReference(const FWorldPartitio
 	}
 }
 
-void UWorldPartitionChangelistValidator::OnInvalidReferenceGridPlacement(const FWorldPartitionActorDescView& ActorDescView, const FWorldPartitionActorDescView& ReferenceActorDescView)
+void UWorldPartitionChangelistValidator::OnInvalidReferenceGridPlacement(const IWorldPartitionActorDescInstanceView& ActorDescView, const IWorldPartitionActorDescInstanceView& ReferenceActorDescView)
 {
 	if (Filter(ActorDescView) || Filter(ReferenceActorDescView))
 	{
@@ -296,7 +295,7 @@ void UWorldPartitionChangelistValidator::OnInvalidReferenceGridPlacement(const F
 	}
 }
 
-void UWorldPartitionChangelistValidator::OnInvalidReferenceDataLayers(const FWorldPartitionActorDescView& ActorDescView, const FWorldPartitionActorDescView& ReferenceActorDescView)
+void UWorldPartitionChangelistValidator::OnInvalidReferenceDataLayers(const IWorldPartitionActorDescInstanceView& ActorDescView, const IWorldPartitionActorDescInstanceView& ReferenceActorDescView)
 {	
 	if (Filter(ActorDescView) || Filter(ReferenceActorDescView))
 	{
@@ -308,7 +307,7 @@ void UWorldPartitionChangelistValidator::OnInvalidReferenceDataLayers(const FWor
 	}
 }
 
-void UWorldPartitionChangelistValidator::OnInvalidReferenceRuntimeGrid(const FWorldPartitionActorDescView& ActorDescView, const FWorldPartitionActorDescView& ReferenceActorDescView)
+void UWorldPartitionChangelistValidator::OnInvalidReferenceRuntimeGrid(const IWorldPartitionActorDescInstanceView& ActorDescView, const IWorldPartitionActorDescInstanceView& ReferenceActorDescView)
 {
 	if (Filter(ActorDescView) || Filter(ReferenceActorDescView))
 	{
@@ -320,7 +319,7 @@ void UWorldPartitionChangelistValidator::OnInvalidReferenceRuntimeGrid(const FWo
 	}
 }
 
-void UWorldPartitionChangelistValidator::OnInvalidReferenceLevelScriptStreamed(const FWorldPartitionActorDescView& ActorDescView)
+void UWorldPartitionChangelistValidator::OnInvalidReferenceLevelScriptStreamed(const IWorldPartitionActorDescInstanceView& ActorDescView)
 {
 	if (Filter(ActorDescView))
 	{		
@@ -331,7 +330,7 @@ void UWorldPartitionChangelistValidator::OnInvalidReferenceLevelScriptStreamed(c
 	}
 }
 
-void UWorldPartitionChangelistValidator::OnInvalidReferenceLevelScriptDataLayers(const FWorldPartitionActorDescView& ActorDescView)
+void UWorldPartitionChangelistValidator::OnInvalidReferenceLevelScriptDataLayers(const IWorldPartitionActorDescInstanceView& ActorDescView)
 {
 	if (Filter(ActorDescView))
 	{
@@ -384,12 +383,12 @@ void UWorldPartitionChangelistValidator::OnDataLayerAssetConflict(const UDataLay
 	}
 }
 
-void UWorldPartitionChangelistValidator::OnActorNeedsResave(const FWorldPartitionActorDescView& ActorDescView)
+void UWorldPartitionChangelistValidator::OnActorNeedsResave(const IWorldPartitionActorDescInstanceView& ActorDescView)
 {
 	// Changelist validation already ensures that dirty actors must be part of the changelist
 }
 
-void UWorldPartitionChangelistValidator::OnLevelInstanceInvalidWorldAsset(const FWorldPartitionActorDescView& ActorDescView, FName WorldAsset, ELevelInstanceInvalidReason Reason)
+void UWorldPartitionChangelistValidator::OnLevelInstanceInvalidWorldAsset(const IWorldPartitionActorDescInstanceView& ActorDescView, FName WorldAsset, ELevelInstanceInvalidReason Reason)
 {
 	if (Filter(ActorDescView))
 	{
@@ -422,12 +421,12 @@ void UWorldPartitionChangelistValidator::OnLevelInstanceInvalidWorldAsset(const
 	}
 }
 
-void UWorldPartitionChangelistValidator::OnInvalidActorFilterReference(const FWorldPartitionActorDescView& ActorDescView, const FWorldPartitionActorDescView& ReferenceActorDescView)
+void UWorldPartitionChangelistValidator::OnInvalidActorFilterReference(const IWorldPartitionActorDescInstanceView& ActorDescView, const IWorldPartitionActorDescInstanceView& ReferenceActorDescView)
 {
 	// Not a validation error
 }
 
-void UWorldPartitionChangelistValidator::OnInvalidHLODLayer(const FWorldPartitionActorDescView& ActorDescView)
+void UWorldPartitionChangelistValidator::OnInvalidHLODLayer(const IWorldPartitionActorDescInstanceView& ActorDescView)
 {
 	if (Filter(ActorDescView))
 	{
diff --git a/Engine/Plugins/Editor/DataValidation/Source/DataValidation/Private/WorldPartitionChangelistValidator.h b/Engine/Plugins/Editor/DataValidation/Source/DataValidation/Private/WorldPartitionChangelistValidator.h
--- a/Engine/Plugins/Editor/DataValidation/Source/DataValidation/Private/WorldPartitionChangelistValidator.h
+++ b/Engine/Plugins/Editor/DataValidation/Source/DataValidation/Private/WorldPartitionChangelistValidator.h
@@ -7,7 +7,7 @@
 #include "WorldPartitionChangelistValidator.generated.h"
 
 class UDataValidationChangelist;
-class UActorDescContainer;
+class UActorDescContainerInstance;
 class FWorldPartitionActorDesc;
 
 UCLASS()
@@ -28,24 +28,24 @@ protected:
 	void ValidateActorsAndDataLayersFromChangeList(UDataValidationChangelist* Changelist);
 
 	// Return true if this ActorDescView is pertinent to the current ChangeList
-	bool Filter(const FWorldPartitionActorDescView& ActorDescView);
+	bool Filter(const IWorldPartitionActorDescInstanceView& ActorDescView);
 
 	// Return true if this UDataLayerInstance is pertinent to the current ChangeList
 	bool Filter(const UDataLayerInstance* DataLayerInstance);
 
 	// IStreamingGenerationErrorHandler Interface methods
-	virtual void OnInvalidRuntimeGrid(const FWorldPartitionActorDescView& ActorDescView, FName GridName) override;
-	virtual void OnInvalidReference(const FWorldPartitionActorDescView& ActorDescView, const FGuid& ReferenceGuid, FWorldPartitionActorDescView* ReferenceActorDescView) override;
-	virtual void OnInvalidReferenceGridPlacement(const FWorldPartitionActorDescView& ActorDescView, const FWorldPartitionActorDescView& ReferenceActorDescView) override;
-	virtual void OnInvalidReferenceDataLayers(const FWorldPartitionActorDescView& ActorDescView, const FWorldPartitionActorDescView& ReferenceActorDescView) override;
-	virtual void OnInvalidReferenceLevelScriptStreamed(const FWorldPartitionActorDescView& ActorDescView) override;
-	virtual void OnInvalidReferenceLevelScriptDataLayers(const FWorldPartitionActorDescView& ActorDescView) override;
-	virtual void OnInvalidReferenceRuntimeGrid(const FWorldPartitionActorDescView& ActorDescView, const FWorldPartitionActorDescView& ReferenceActorDescView) override;
+	virtual void OnInvalidRuntimeGrid(const IWorldPartitionActorDescInstanceView& ActorDescView, FName GridName) override;
+	virtual void OnInvalidReference(const IWorldPartitionActorDescInstanceView& ActorDescView, const FGuid& ReferenceGuid, IWorldPartitionActorDescInstanceView* ReferenceActorDescView) override;
+	virtual void OnInvalidReferenceGridPlacement(const IWorldPartitionActorDescInstanceView& ActorDescView, const IWorldPartitionActorDescInstanceView& ReferenceActorDescView) override;
+	virtual void OnInvalidReferenceDataLayers(const IWorldPartitionActorDescInstanceView& ActorDescView, const IWorldPartitionActorDescInstanceView& ReferenceActorDescView) override;
+	virtual void OnInvalidReferenceLevelScriptStreamed(const IWorldPartitionActorDescInstanceView& ActorDescView) override;
+	virtual void OnInvalidReferenceLevelScriptDataLayers(const IWorldPartitionActorDescInstanceView& ActorDescView) override;
+	virtual void OnInvalidReferenceRuntimeGrid(const IWorldPartitionActorDescInstanceView& ActorDescView, const IWorldPartitionActorDescInstanceView& ReferenceActorDescView) override;
 	virtual void OnInvalidReferenceDataLayerAsset(const UDataLayerInstanceWithAsset* DataLayerInstance) override;
 	virtual void OnDataLayerHierarchyTypeMismatch(const UDataLayerInstance* DataLayerInstance, const UDataLayerInstance* Parent) override;
 	virtual void OnDataLayerAssetConflict(const UDataLayerInstanceWithAsset* DataLayerInstance, const UDataLayerInstanceWithAsset* ConflictingDataLayerInstance) override;
-	virtual void OnActorNeedsResave(const FWorldPartitionActorDescView& ActorDescView) override;
-	virtual void OnLevelInstanceInvalidWorldAsset(const FWorldPartitionActorDescView& ActorDescView, FName WorldAsset, ELevelInstanceInvalidReason Reason) override;
-	virtual void OnInvalidActorFilterReference(const FWorldPartitionActorDescView& ActorDescView, const FWorldPartitionActorDescView& ReferenceActorDescView) override;
-	virtual void OnInvalidHLODLayer(const FWorldPartitionActorDescView& ActorDescView) override;
+	virtual void OnActorNeedsResave(const IWorldPartitionActorDescInstanceView& ActorDescView) override;
+	virtual void OnLevelInstanceInvalidWorldAsset(const IWorldPartitionActorDescInstanceView& ActorDescView, FName WorldAsset, ELevelInstanceInvalidReason Reason) override;
+	virtual void OnInvalidActorFilterReference(const IWorldPartitionActorDescInstanceView& ActorDescView, const IWorldPartitionActorDescInstanceView& ReferenceActorDescView) override;
+	virtual void OnInvalidHLODLayer(const IWorldPartitionActorDescInstanceView& ActorDescView) override;
 };
diff --git a/Engine/Plugins/Editor/ObjectMixer/ObjectMixer/Source/ObjectMixer/Private/Views/List/Modes/ObjectMixerOutlinerHierarchy.cpp b/Engine/Plugins/Editor/ObjectMixer/ObjectMixer/Source/ObjectMixer/Private/Views/List/Modes/ObjectMixerOutlinerHierarchy.cpp
--- a/Engine/Plugins/Editor/ObjectMixer/ObjectMixer/Source/ObjectMixer/Private/Views/List/Modes/ObjectMixerOutlinerHierarchy.cpp
+++ b/Engine/Plugins/Editor/ObjectMixer/ObjectMixer/Source/ObjectMixer/Private/Views/List/Modes/ObjectMixerOutlinerHierarchy.cpp
@@ -27,8 +27,9 @@
 #include "LevelInstance/LevelInstanceSubsystem.h"
 #include "Modules/ModuleManager.h"
 #include "SceneOutlinerFwd.h"
-#include "WorldPartition/ActorDescContainer.h"
-#include "WorldPartition/ActorDescContainerCollection.h"
+#include "WorldPartition/ActorDescContainerInstanceCollection.h"
+#include "WorldPartition/ActorDescContainerInstance.h"
+#include "WorldPartition/WorldPartitionActorDescInstance.h"
 #include "WorldPartition/IWorldPartitionEditorModule.h"
 #include "WorldPartition/WorldPartition.h"
 #include "WorldPartition/WorldPartitionActorDesc.h"
@@ -64,7 +65,7 @@ TUniquePtr<FObjectMixerOutlinerHierarchy> FObjectMixerOutlinerHierarchy::Create(
 
 		if (UWorldPartition* WorldPartition = World->GetWorldPartition())
 		{
-			WorldPartition->OnActorDescRemovedEvent.AddRaw(Hierarchy, &FObjectMixerOutlinerHierarchy::OnActorDescRemoved);
+			WorldPartition->OnActorDescInstanceRemovedEvent.AddRaw(Hierarchy, &FObjectMixerOutlinerHierarchy::OnActorDescInstanceRemoved);
 		}
 	}
 
@@ -117,8 +118,7 @@ FObjectMixerOutlinerHierarchy::~FObjectMixerOutlinerHierarchy()
 
 		if (UWorldPartition* WorldPartition = RepresentingWorld->GetWorldPartition())
 		{
-			WorldPartition->OnActorDescAddedEvent.RemoveAll(this);
-			WorldPartition->OnActorDescRemovedEvent.RemoveAll(this);
+			WorldPartition->OnActorDescInstanceRemovedEvent.RemoveAll(this);
 		}
 	}
 
@@ -312,11 +312,11 @@ FSceneOutlinerTreeItemPtr FObjectMixerOutlinerHierarchy::FindOrCreateParentItem(
 	}
 	else if (const FActorDescTreeItem* ActorDescItem = Item.CastTo<FActorDescTreeItem>())
 	{
-		if (const FWorldPartitionActorDesc* ActorDesc = ActorDescItem->ActorDescHandle.Get())
+		if (const FWorldPartitionActorDescInstance* ActorDescInstance = ActorDescItem->ActorDescHandle.GetInstance())
 		{
 			if (UWorld* RepresentingWorldPtr = RepresentingWorld.Get())
 			{
-				const FFolder ActorDescFolder = FActorFolders::GetActorDescFolder(*RepresentingWorldPtr, ActorDesc);
+				const FFolder ActorDescFolder = FActorFolders::GetActorDescInstanceFolder(*RepresentingWorldPtr, ActorDescInstance);
 				if (!ActorDescFolder.IsNone())
 				{
 					if (const FSceneOutlinerTreeItemPtr* ParentItem = Items.Find(ActorDescFolder))
@@ -331,15 +331,15 @@ FSceneOutlinerTreeItemPtr FObjectMixerOutlinerHierarchy::FindOrCreateParentItem(
 			}
 
 			// Parent Actor (Actor attachment / parenting)
-			const FGuid& ParentActorGuid = ActorDesc->GetSceneOutlinerParent();
+			const FGuid& ParentActorGuid = ActorDescInstance->GetSceneOutlinerParent();
 			if (ParentActorGuid.IsValid())
 			{
-				if (UActorDescContainer* ActorDescContainer = ActorDesc->GetContainer())
+				if (UActorDescContainerInstance* ActorDescContainerInstance = ActorDescInstance->GetContainerInstance())
 				{
-					if (const FWorldPartitionActorDesc* ParentActorDesc = ActorDescContainer->GetActorDesc(ParentActorGuid))
+					if (const FWorldPartitionActorDescInstance* ParentActorDescInstance = ActorDescContainerInstance->GetActorDescInstance(ParentActorGuid))
 					{
 						// If parent actor is loaded
-						if (AActor* ParentActor = ParentActorDesc->GetActor())
+						if (AActor* ParentActor = ParentActorDescInstance->GetActor())
 						{
 							// Find loaded parent actor node (from the object ptr)
 							if (const FSceneOutlinerTreeItemPtr* ParentItem = Items.Find(ParentActor))
@@ -362,21 +362,21 @@ FSceneOutlinerTreeItemPtr FObjectMixerOutlinerHierarchy::FindOrCreateParentItem(
 						}
 
 						// Find unloaded parent actor node (from the guid)
-						if (const FSceneOutlinerTreeItemPtr* ParentItem = Items.Find(FActorDescTreeItem::ComputeTreeItemID(ParentActorGuid, ActorDescContainer)))
+						if (const FSceneOutlinerTreeItemPtr* ParentItem = Items.Find(FActorDescTreeItem::ComputeTreeItemID(ParentActorGuid, ActorDescContainerInstance)))
 						{
 							return *ParentItem;
 						}
 						else
 						{
-							return bCreate ? Mode->CreateItemFor<FActorDescTreeItem>(FActorDescTreeItem(ParentActorGuid, ActorDescContainer)) : nullptr;
+							return bCreate ? Mode->CreateItemFor<FActorDescTreeItem>(FActorDescTreeItem(ParentActorGuid, ActorDescContainerInstance)) : nullptr;
 						}
 					}
 				}
 			}
-			else if (UActorDescContainer* ActorDescContainer = ActorDesc->GetContainer())
+			else if (UActorDescContainerInstance* ActorDescContainerInstance = ActorDescInstance->GetContainerInstance())
 			{
 				const ULevelInstanceSubsystem* LevelInstanceSubsystem = UWorld::GetSubsystem<ULevelInstanceSubsystem>(RepresentingWorld.Get());
-				UWorld* OuterWorld = ActorDescContainer->GetTypedOuter<UWorld>();
+				UWorld* OuterWorld = ActorDescContainerInstance->GetTypedOuter<UWorld>();
 				// If parent actor is loaded
 				if (AActor* ParentActor = OuterWorld ? Cast<AActor>(LevelInstanceSubsystem->GetOwningLevelInstance(OuterWorld->PersistentLevel)) : nullptr)
 				{
@@ -522,13 +522,13 @@ void FObjectMixerOutlinerHierarchy::CreateWorldChildren(UWorld* World, TArray<FS
 						return true;
 					}
 				}
-				FWorldPartitionHelpers::ForEachActorDesc(WorldPartition,
-					[this, &OutItems](const FWorldPartitionActorDesc* ActorDesc)
+				FWorldPartitionHelpers::ForEachActorDescInstance(WorldPartition,
+					[this, &OutItems](const FWorldPartitionActorDescInstance* ActorDescInstance)
 				{
-					if (ActorDesc != nullptr && !ActorDesc->IsLoaded(true))
+					if (ActorDescInstance != nullptr && !ActorDescInstance->IsLoaded(true))
 					{
 						if (const FSceneOutlinerTreeItemPtr ActorDescItem =
-							Mode->CreateItemFor<FActorDescTreeItem>(FActorDescTreeItem(ActorDesc->GetGuid(), ActorDesc->GetContainer())))
+							Mode->CreateItemFor<FActorDescTreeItem>(FActorDescTreeItem(ActorDescInstance->GetGuid(), ActorDescInstance->GetContainerInstance())))
 						{
 							OutItems.Add(ActorDescItem);
 						}
@@ -813,11 +813,11 @@ void FObjectMixerOutlinerHierarchy::OnLoadedActorAdded(AActor& InActor)
 	if (UWorldPartition* WorldPartition = FWorldPartitionHelpers::GetWorldPartition(&InActor))
 	{
 		const FGuid& ActorGuid = InActor.GetActorGuid();
-		if (FWorldPartitionActorDesc* ActorDesc = WorldPartition->GetActorDesc(ActorGuid))
+		if (FWorldPartitionActorDescInstance* ActorDescInstance = WorldPartition->GetActorDescInstance(ActorGuid))
 		{
 			FSceneOutlinerHierarchyChangedData EventData;
 			EventData.Type = FSceneOutlinerHierarchyChangedData::Removed;
-			EventData.ItemIDs.Add(FActorDescTreeItem::ComputeTreeItemID(ActorDesc->GetGuid(), ActorDesc->GetContainer()));
+			EventData.ItemIDs.Add(FActorDescTreeItem::ComputeTreeItemID(ActorDescInstance->GetGuid(), ActorDescInstance->GetContainerInstance()));
 			HierarchyChangedEvent.Broadcast(EventData);
 		}
 	}
@@ -828,13 +828,13 @@ void FObjectMixerOutlinerHierarchy::OnLoadedActorRemoved(AActor& InActor)
 	OnLevelActorDeleted(&InActor);
 }
 
-void FObjectMixerOutlinerHierarchy::OnActorDescRemoved(FWorldPartitionActorDesc* ActorDesc)
+void FObjectMixerOutlinerHierarchy::OnActorDescInstanceRemoved(FWorldPartitionActorDescInstance* ActorDescInstance)
 {
-	if (bShowingUnloadedActors && ActorDesc)
+	if (bShowingUnloadedActors && ActorDescInstance)
 	{
 		FSceneOutlinerHierarchyChangedData EventData;
 		EventData.Type = FSceneOutlinerHierarchyChangedData::Removed;
-		EventData.ItemIDs.Add(FActorDescTreeItem::ComputeTreeItemID(ActorDesc->GetGuid(), ActorDesc->GetContainer()));
+		EventData.ItemIDs.Add(FActorDescTreeItem::ComputeTreeItemID(ActorDescInstance->GetGuid(), ActorDescInstance->GetContainerInstance()));
 		HierarchyChangedEvent.Broadcast(EventData);
 	}
 }
diff --git a/Engine/Plugins/Editor/ObjectMixer/ObjectMixer/Source/ObjectMixer/Private/Views/List/Modes/ObjectMixerOutlinerMode.cpp b/Engine/Plugins/Editor/ObjectMixer/ObjectMixer/Source/ObjectMixer/Private/Views/List/Modes/ObjectMixerOutlinerMode.cpp
--- a/Engine/Plugins/Editor/ObjectMixer/ObjectMixer/Source/ObjectMixer/Private/Views/List/Modes/ObjectMixerOutlinerMode.cpp
+++ b/Engine/Plugins/Editor/ObjectMixer/ObjectMixer/Source/ObjectMixer/Private/Views/List/Modes/ObjectMixerOutlinerMode.cpp
@@ -60,6 +60,7 @@
 #include "WorldPartition/IWorldPartitionEditorModule.h"
 #include "WorldPartition/WorldPartition.h"
 #include "WorldPartition/WorldPartitionSubsystem.h"
+#include "WorldPartition/WorldPartitionActorDescInstance.h"
 
 static int32 GSceneOutlinerAutoRepresentingWorldNetModeForObjectMixer = NM_Client;
 static FAutoConsoleVariableRef CVarAutoRepresentingWorldNetMode(
@@ -172,17 +173,14 @@ namespace ObjectMixerOutliner
 		return false;
 	}
 
-	bool FActorDescSelector::operator()(const TWeakPtr<ISceneOutlinerTreeItem>& Item, FWorldPartitionActorDesc*& ActorDescPtrOut) const
+	bool FActorHandleSelector::operator()(const TWeakPtr<ISceneOutlinerTreeItem>& Item, FWorldPartitionHandle& ActorHandleOut) const
 	{
 		if (TSharedPtr<ISceneOutlinerTreeItem> ItemPtr = Item.Pin())
 		{
 			if (FActorDescTreeItem* ActorDescItem = ItemPtr->CastTo<FActorDescTreeItem>())
 			{
-				if (FWorldPartitionActorDesc* ActorDesc = ActorDescItem->ActorDescHandle.Get())
-				{
-					ActorDescPtrOut = ActorDesc;
-					return true;
-				}
+				ActorHandleOut = ActorDescItem->ActorDescHandle;
+				return true;
 			}
 		}
 
@@ -316,7 +314,7 @@ FObjectMixerOutlinerMode::FObjectMixerOutlinerMode(
 	FilterInfoMap.Add(TEXT("ShowOnlyCurrentDataLayers"), OnlyCurrentDataLayersInfo);
 
 	// Add a filter for unloaded actors to properly reflect the bShowOnlyActorsInCurrentDataLayers flag.
-	SceneOutliner->AddFilter(MakeShared<FActorDescFilter>(FActorDescTreeItem::FFilterPredicate::CreateLambda([this](const FWorldPartitionActorDesc* ActorDesc)
+	SceneOutliner->AddFilter(MakeShared<FActorDescFilter>(FActorDescTreeItem::FFilterPredicate::CreateLambda([this](const FWorldPartitionActorDescInstance* ActorDescInstance)
 		{
 			FObjectMixerOutlinerModeConfig* Settings = GetMutableConfig();
 			if (Settings && Settings->bShowOnlyActorsInCurrentDataLayers)
@@ -344,7 +342,7 @@ FObjectMixerOutlinerMode::FObjectMixerOutlinerMode(
 	FilterInfoMap.Add(TEXT("ShowOnlyCurrentContentBundle"), OnlyCurrentContentBundleInfo);
 
 	// Add a filter for unloaded actors to properly reflect the bShowOnlyActorsInCurrentContentBundle flag.
-	SceneOutliner->AddFilter(MakeShared<FActorDescFilter>(FActorDescTreeItem::FFilterPredicate::CreateLambda([this](const FWorldPartitionActorDesc* ActorDesc)
+	SceneOutliner->AddFilter(MakeShared<FActorDescFilter>(FActorDescTreeItem::FFilterPredicate::CreateLambda([this](const FWorldPartitionActorDescInstance* ActorDescInstance)
 		{
 			FObjectMixerOutlinerModeConfig* Settings = GetMutableConfig();
 			if (Settings && Settings->bShowOnlyActorsInCurrentContentBundle)
@@ -354,7 +352,7 @@ FObjectMixerOutlinerMode::FObjectMixerOutlinerMode(
 					return true;
 				}
 
-				return ActorDesc->GetContentBundleGuid().IsValid() && WorldPartitionEditorModule->IsEditingContentBundle(ActorDesc->GetContentBundleGuid());
+				return ActorDescInstance->GetContentBundleGuid().IsValid() && WorldPartitionEditorModule->IsEditingContentBundle(ActorDescInstance->GetContentBundleGuid());
 			}
 			return true;
 		}), FSceneOutlinerFilter::EDefaultBehaviour::Pass));
@@ -500,7 +498,7 @@ FObjectMixerOutlinerMode::~FObjectMixerOutlinerMode()
 	{
 		if (UWorldPartition* const WorldPartition = RepresentingWorld->GetWorldPartition())
 		{
-			WorldPartition->OnActorDescRemovedEvent.RemoveAll(this);
+			WorldPartition->OnActorDescInstanceRemovedEvent.RemoveAll(this);
 		}
 	}
 	FSceneOutlinerDelegates::Get().OnComponentsUpdated.RemoveAll(this);
@@ -921,7 +919,7 @@ void FObjectMixerOutlinerMode::Rebuild()
 	{
 		if (UWorldPartition* const WorldPartition = RepresentingWorld->GetWorldPartition())
 		{
-			WorldPartition->OnActorDescRemovedEvent.RemoveAll(this);
+			WorldPartition->OnActorDescInstanceRemovedEvent.RemoveAll(this);
 		}
 	}
 
@@ -939,7 +937,7 @@ void FObjectMixerOutlinerMode::Rebuild()
 	if (bRepresentingWorldPartitionedWorld)
 	{
 		UWorldPartition* const WorldPartition = RepresentingWorld->GetWorldPartition();
-		WorldPartition->OnActorDescRemovedEvent.AddRaw(this, &FObjectMixerOutlinerMode::OnActorDescRemoved);
+		WorldPartition->OnActorDescInstanceRemovedEvent.AddRaw(this, &FObjectMixerOutlinerMode::OnActorDescInstanceRemoved);
 	}
 
 	SceneOutliner->FullRefresh();
@@ -1433,7 +1431,7 @@ TSharedRef<FSceneOutlinerFilter> FObjectMixerOutlinerMode::CreateHideLevelInstan
 TSharedRef<FSceneOutlinerFilter> FObjectMixerOutlinerMode::CreateHideUnloadedActorsFilter()
 {
 	return MakeShareable(new FActorDescFilter(FActorDescTreeItem::FFilterPredicate::CreateStatic(
-		[](const FWorldPartitionActorDesc* ActorDesc) { return false; }), FSceneOutlinerFilter::EDefaultBehaviour::Pass));
+		[](const FWorldPartitionActorDescInstance* ActorDescInstance) { return false; }), FSceneOutlinerFilter::EDefaultBehaviour::Pass));
 }
 
 TSharedRef<FSceneOutlinerFilter> FObjectMixerOutlinerMode::CreateHideEmptyFoldersFilter()
@@ -1806,9 +1804,9 @@ void FObjectMixerOutlinerMode::OnSelectUnloadedActors(const TArray<FGuid>& Actor
 		ItemsToSelect.Reserve(ActorGuids.Num());
 		for (const FGuid& ActorGuid : ActorGuids)
 		{
-			if (FWorldPartitionActorDesc* ActorDesc = WorldPartition->GetActorDesc(ActorGuid))
+			if (FWorldPartitionActorDescInstance* ActorDescInstance = WorldPartition->GetActorDescInstance(ActorGuid))
 			{
-				if (FSceneOutlinerTreeItemPtr ItemPtr = SceneOutliner->GetTreeItem(FActorDescTreeItem::ComputeTreeItemID(ActorDesc->GetGuid(), ActorDesc->GetContainer())))
+				if (FSceneOutlinerTreeItemPtr ItemPtr = SceneOutliner->GetTreeItem(FActorDescTreeItem::ComputeTreeItemID(ActorDescInstance->GetGuid(), ActorDescInstance->GetContainerInstance())))
 				{
 					ItemsToSelect.Add(ItemPtr);
 				}
@@ -1828,9 +1826,9 @@ void FObjectMixerOutlinerMode::OnSelectUnloadedActors(const TArray<FGuid>& Actor
 	}
 }
 
-void FObjectMixerOutlinerMode::OnActorDescRemoved(FWorldPartitionActorDesc* InActorDesc)
+void FObjectMixerOutlinerMode::OnActorDescInstanceRemoved(FWorldPartitionActorDescInstance* InActorDescInstance)
 {
-	ApplicableUnloadedActors.Remove(InActorDesc);
+	ApplicableUnloadedActors.Remove(InActorDescInstance);
 }
 
 void FObjectMixerOutlinerMode::OnItemSelectionChanged(FSceneOutlinerTreeItemPtr TreeItem, ESelectInfo::Type SelectionType, const FSceneOutlinerItemSelection& Selection)
@@ -1923,7 +1921,7 @@ void FObjectMixerOutlinerMode::OnItemPassesFilters(const ISceneOutlinerTreeItem&
 	}
 	else if (const FActorDescTreeItem* const ActorDescItem = Item.CastTo<FActorDescTreeItem>(); ActorDescItem && ActorDescItem->IsValid())
 	{
-		ApplicableUnloadedActors.Add(ActorDescItem->ActorDescHandle.Get());
+		ApplicableUnloadedActors.Add(ActorDescItem->ActorDescHandle);
 	}
 }
 
@@ -2050,9 +2048,9 @@ bool FObjectMixerOutlinerMode::HasErrors() const
 		{
 			bool bHasErrors = 0;
 
-			WorldPartition->ForEachActorDescContainer([&bHasErrors](UActorDescContainer* ActorDescContainer)
+			WorldPartition->ForEachActorDescContainer([&bHasErrors](UActorDescContainerInstance* ActorDescContainerInstance)
 			{
-				if (ActorDescContainer->HasInvalidActors())
+				if (ActorDescContainerInstance->GetContainer()->HasInvalidActors())
 				{
 					bHasErrors = true;
 				}
@@ -2080,13 +2078,13 @@ void FObjectMixerOutlinerMode::RepairErrors() const
 			ISourceControlProvider& SourceControlProvider = SourceControlModule.GetProvider();
 
 			TArray<FAssetData> InvalidActorAssets;
-			WorldPartition->ForEachActorDescContainer([&InvalidActorAssets](UActorDescContainer* ActorDescContainer)
+			WorldPartition->ForEachActorDescContainer([&InvalidActorAssets](UActorDescContainerInstance* ActorDescContainerInstance)
 			{
-				for (const FAssetData& InvalidActor : ActorDescContainer->GetInvalidActors())
+				for (const FAssetData& InvalidActor : ActorDescContainerInstance->GetContainer()->GetInvalidActors())
 				{
 					InvalidActorAssets.Add(InvalidActor);
 				}
-				ActorDescContainer->ClearInvalidActors();
+				ActorDescContainerInstance->GetContainer()->ClearInvalidActors();
 			});
 
 			TArray<FString> ActorFilesToDelete;
@@ -2281,12 +2279,12 @@ void FObjectMixerOutlinerMode::SynchronizeSelectedActorDescs()
 	if (UWorldPartitionSubsystem* WorldPartitionSubsystem = UWorld::GetSubsystem<UWorldPartitionSubsystem>(RepresentingWorld.Get()))
 	{
 		const FSceneOutlinerItemSelection Selection = SceneOutliner->GetSelection();
-		TArray<FWorldPartitionActorDesc*> SelectedActorDescs = Selection.GetData<FWorldPartitionActorDesc*>(ObjectMixerOutliner::FActorDescSelector());
+		TArray<FWorldPartitionHandle> SelectedActorHandles = Selection.GetData<FWorldPartitionHandle>(ObjectMixerOutliner::FActorHandleSelector());
 
-		WorldPartitionSubsystem->SelectedActorDescs.Empty();
-		for (FWorldPartitionActorDesc* SelectedActorDesc : SelectedActorDescs)
+		WorldPartitionSubsystem->SelectedActorHandles.Empty();
+		for (const FWorldPartitionHandle& ActorHandle : SelectedActorHandles)
 		{
-			WorldPartitionSubsystem->SelectedActorDescs.Add(SelectedActorDesc);
+			WorldPartitionSubsystem->SelectedActorHandles.Add(ActorHandle);
 		}
 	}
 }
@@ -2988,11 +2986,11 @@ namespace ObjectMixerActorBrowsingModeUtils
 		
 		for (const TTuple<UWorld*, TSet<FName>>& Pair : UnloadedActorsFolderPaths)
 		{
-			FActorFolders::ForEachActorDescInFolders(*Pair.Key, Pair.Value, [&List](const FWorldPartitionActorDesc* ActorDesc)
+			FActorFolders::ForEachActorDescInstanceInFolders(*Pair.Key, Pair.Value, [&List](const FWorldPartitionActorDescInstance* ActorDescInstance)
 			{
-				if (!ActorDesc->IsLoaded())
+				if (!ActorDescInstance->IsLoaded())
 				{
-					List.Add(ActorDesc->GetGuid());
+					List.Add(ActorDescInstance->GetGuid());
 				}
 				return true;
 			});
@@ -3072,7 +3070,7 @@ void FObjectMixerOutlinerMode::PinItems(const TArray<FSceneOutlinerTreeItemPtr>&
 		{
 			if (FWorldPartitionHandle ActorHandle(WorldPartition, ActorGuid); ActorHandle.IsValid())
 			{
-				if (AActor* PinnedActor = ActorHandle->GetActor())
+				if (AActor* PinnedActor = ActorHandle.GetActor())
 				{
 					GEditor->SelectActor(PinnedActor, /*bInSelected=*/true, /*bNotify=*/false);
 					LastPinnedActor = PinnedActor;
@@ -3115,7 +3113,7 @@ void FObjectMixerOutlinerMode::UnpinItems(const TArray<FSceneOutlinerTreeItemPtr
 		{
 			if (FWorldPartitionHandle ActorHandle(WorldPartition, ActorGuid); ActorHandle.IsValid())
 			{
-				if (AActor* PinnedActor = ActorHandle->GetActor())
+				if (AActor* PinnedActor = ActorHandle.GetActor())
 				{
 					GEditor->SelectActor(PinnedActor, /*bInSelected=*/false, /*bNotify=*/false);
 				}
@@ -3297,10 +3295,10 @@ bool FObjectMixerOutlinerMode::CompareItemWithClassName(SceneOutliner::FilterBar
 	}
 	else if (const FActorDescTreeItem* ActorDescItem = InItem.CastTo<FActorDescTreeItem>())
 	{
-		if (const FWorldPartitionActorDesc* ActorDesc = ActorDescItem->ActorDescHandle.Get())
+		if (const FWorldPartitionActorDescInstance* ActorDescInstance = ActorDescItem->ActorDescHandle.GetInstance())
 		{
 			// For Unloaded Actors, grab the native class 
-			FTopLevelAssetPath ClassPath = ActorDesc->GetNativeClass();
+			FTopLevelAssetPath ClassPath = ActorDescInstance->GetNativeClass();
 			return AssetClassPaths.Contains(ClassPath);
 		}
 	}
diff --git a/Engine/Plugins/Editor/ObjectMixer/ObjectMixer/Source/ObjectMixer/Private/Views/List/SObjectMixerEditorList.cpp b/Engine/Plugins/Editor/ObjectMixer/ObjectMixer/Source/ObjectMixer/Private/Views/List/SObjectMixerEditorList.cpp
--- a/Engine/Plugins/Editor/ObjectMixer/ObjectMixer/Source/ObjectMixer/Private/Views/List/SObjectMixerEditorList.cpp
+++ b/Engine/Plugins/Editor/ObjectMixer/ObjectMixer/Source/ObjectMixer/Private/Views/List/SObjectMixerEditorList.cpp
@@ -39,6 +39,7 @@
 #include "WorldPartition/ContentBundle/ContentBundleDescriptor.h"
 #include "WorldPartition/ContentBundle/ContentBundleEngineSubsystem.h"
 #include "WorldPartition/DataLayer/DataLayerManager.h"
+#include "WorldPartition/WorldPartitionActorDescInstance.h"
 
 #define LOCTEXT_NAMESPACE "SObjectMixerEditorList"
 
@@ -902,18 +903,18 @@ void SObjectMixerEditorList::CreateActorTextInfoColumns(FSceneOutlinerInitializa
 			}
 			else if (const FActorDescTreeItem* ActorDescItem = Item.CastTo<FActorDescTreeItem>())
 			{
-				if (const FWorldPartitionActorDesc* ActorDesc = ActorDescItem->ActorDescHandle.Get(); ActorDesc && !ActorDesc->GetDataLayerInstanceNames().IsEmpty())
+				if (const FWorldPartitionActorDescInstance* ActorDescInstance = ActorDescItem->ActorDescHandle.GetInstance(); ActorDescInstance && !ActorDescInstance->GetDataLayerInstanceNames().IsEmpty())
 				{
-					if (const UActorDescContainer* ActorDescContainer = ActorDescItem->ActorDescHandle.Container.Get())
+					if (const UActorDescContainerInstance* ActorDescContainerInstance = ActorDescInstance->GetContainerInstance())
 					{
-						const UWorld* OwningWorld = ActorDescContainer->GetWorldPartition()->GetWorld();
+						const UWorld* OwningWorld = ActorDescContainerInstance->GetWorldPartition()->GetWorld();
 						if (const UDataLayerManager* DataLayerManager = UDataLayerManager::GetDataLayerManager(OwningWorld))
 						{
 							TSet<const UDataLayerInstance*> DataLayerInstances;
-							DataLayerInstances.Append(DataLayerManager->GetDataLayerInstances(ActorDesc->GetDataLayerInstanceNames()));
+							DataLayerInstances.Append(DataLayerManager->GetDataLayerInstances(ActorDescInstance->GetDataLayerInstanceNames()));
 							if (ULevelInstanceSubsystem* LevelInstanceSubsystem = UWorld::GetSubsystem<ULevelInstanceSubsystem>(OwningWorld))
 							{
-								UWorld* OuterWorld = ActorDescContainer->GetTypedOuter<UWorld>();
+								UWorld* OuterWorld = ActorDescContainerInstance->GetTypedOuter<UWorld>();
 								// Add parent container Data Layer Instances
 								AActor* CurrentActor = OuterWorld ? Cast<AActor>(LevelInstanceSubsystem->GetOwningLevelInstance(OuterWorld->PersistentLevel)) : nullptr;
 								while (CurrentActor)
@@ -951,9 +952,9 @@ void SObjectMixerEditorList::CreateActorTextInfoColumns(FSceneOutlinerInitializa
 		}
 		else if (const FActorDescTreeItem* ActorDescItem = Item.CastTo<FActorDescTreeItem>())
 		{
-			if (const FWorldPartitionActorDesc* ActorDesc = ActorDescItem->ActorDescHandle.Get())
+			if (const FWorldPartitionActorDescInstance* ActorDescInstance = ActorDescItem->ActorDescHandle.GetInstance())
 			{
-				if (const UContentBundleDescriptor* Descriptor = ContentBundleEngineSubsystem->GetContentBundleDescriptor(ActorDesc->GetContentBundleGuid()))
+				if (const UContentBundleDescriptor* Descriptor = ContentBundleEngineSubsystem->GetContentBundleDescriptor(ActorDescInstance->GetContentBundleGuid()))
 				{
 					return Descriptor->GetDisplayName();
 				}
@@ -977,11 +978,11 @@ void SObjectMixerEditorList::CreateActorTextInfoColumns(FSceneOutlinerInitializa
 		}
 		else if (const FActorDescTreeItem* ActorDescItem = Item.CastTo<FActorDescTreeItem>())
 		{
-			if (const FWorldPartitionActorDesc* ActorDesc = ActorDescItem->ActorDescHandle.Get())
+			if (const FWorldPartitionActorDescInstance* ActorDescInstance = ActorDescItem->ActorDescHandle.GetInstance())
 			{
-				if (FName LevelPackage = ActorDesc->GetContainerPackage(); !LevelPackage.IsNone())
+				if (FName LevelPackage = ActorDescInstance->GetChildContainerPackage(); !LevelPackage.IsNone())
 				{
-					return ActorDesc->GetContainerPackage().ToString();
+					return ActorDescInstance->GetChildContainerPackage().ToString();
 				}
 			}
 		}
@@ -1022,9 +1023,9 @@ void SObjectMixerEditorList::CreateActorTextInfoColumns(FSceneOutlinerInitializa
 		}
 		else if (const FActorDescTreeItem* ActorDescItem = Item.CastTo<FActorDescTreeItem>())
 		{
-			if (const FWorldPartitionActorDesc* ActorDesc = ActorDescItem->ActorDescHandle.Get())
+			if (const FWorldPartitionActorDescInstance* ActorDescInstance = ActorDescItem->ActorDescHandle.GetInstance())
 			{
-				return ActorDesc->GetActorName().ToString();
+				return ActorDescInstance->GetActorName().ToString();
 			}
 		}
 		else if (const FActorFolderTreeItem* ActorFolderItem = Item.CastTo<FActorFolderTreeItem>())
@@ -1049,9 +1050,9 @@ void SObjectMixerEditorList::CreateActorTextInfoColumns(FSceneOutlinerInitializa
 		}
 		else if (const FActorDescTreeItem* ActorDescItem = Item.CastTo<FActorDescTreeItem>())
 		{
-			if (const FWorldPartitionActorDesc* ActorDesc = ActorDescItem->ActorDescHandle.Get())
+			if (const FWorldPartitionActorDescInstance* ActorDescInstance = ActorDescItem->ActorDescHandle.GetInstance())
 			{
-				return FPackageName::GetShortName(ActorDesc->GetActorPackage());
+				return FPackageName::GetShortName(ActorDescInstance->GetActorPackage());
 			}
 		}
 		else if (const FActorFolderTreeItem* ActorFolderItem = Item.CastTo<FActorFolderTreeItem>())
diff --git a/Engine/Plugins/Editor/ObjectMixer/ObjectMixer/Source/ObjectMixer/Public/Views/List/Modes/ObjectMixerOutlinerHierarchy.h b/Engine/Plugins/Editor/ObjectMixer/ObjectMixer/Source/ObjectMixer/Public/Views/List/Modes/ObjectMixerOutlinerHierarchy.h
--- a/Engine/Plugins/Editor/ObjectMixer/ObjectMixer/Source/ObjectMixer/Public/Views/List/Modes/ObjectMixerOutlinerHierarchy.h
+++ b/Engine/Plugins/Editor/ObjectMixer/ObjectMixer/Source/ObjectMixer/Public/Views/List/Modes/ObjectMixerOutlinerHierarchy.h
@@ -5,7 +5,7 @@
 #include "ISceneOutlinerHierarchy.h"
 
 class FObjectMixerOutlinerMode;
-class FWorldPartitionActorDesc;
+class FWorldPartitionActorDescInstance;
 class UActorFolder;
 class ULevelInstanceSubsystem;
 class UWorldPartition;
@@ -68,7 +68,7 @@ protected:
 	void OnLoadedActorAdded(AActor& InActor);
 	void OnLoadedActorRemoved(AActor& InActor);
 
-	void OnActorDescRemoved(FWorldPartitionActorDesc* ActorDesc);
+	void OnActorDescInstanceRemoved(FWorldPartitionActorDescInstance* ActorDescInstance);
 	
 	void OnComponentsUpdated();
 
diff --git a/Engine/Plugins/Editor/ObjectMixer/ObjectMixer/Source/ObjectMixer/Public/Views/List/Modes/ObjectMixerOutlinerMode.h b/Engine/Plugins/Editor/ObjectMixer/ObjectMixer/Source/ObjectMixer/Public/Views/List/Modes/ObjectMixerOutlinerMode.h
--- a/Engine/Plugins/Editor/ObjectMixer/ObjectMixer/Source/ObjectMixer/Public/Views/List/Modes/ObjectMixerOutlinerMode.h
+++ b/Engine/Plugins/Editor/ObjectMixer/ObjectMixer/Source/ObjectMixer/Public/Views/List/Modes/ObjectMixerOutlinerMode.h
@@ -8,6 +8,7 @@
 #include "ObjectFilter/ObjectMixerEditorObjectFilter.h"
 #include "UObject/ObjectPtr.h"
 #include "UObject/UObjectGlobals.h"
+#include "WorldPartition/WorldPartitionHandle.h"
 
 #include "ObjectMixerOutlinerMode.generated.h"
 
@@ -37,11 +38,16 @@ namespace ObjectMixerOutliner
 	{
 		bool operator()(const TWeakPtr<ISceneOutlinerTreeItem>& Item, AActor*& ActorPtrOut) const;
 	};
+		
+	struct UE_DEPRECATED(5.4, "Use FActorHandleSelector instead") FActorDescSelector
+	{
+		bool operator()(const TWeakPtr<ISceneOutlinerTreeItem>& Item, FWorldPartitionActorDesc*& ActorDescPtrOut) const { return false; }
+	};
 
 	/** Functor which can be used to get actor descriptors from a selection  */
-	struct FActorDescSelector
+	struct FActorHandleSelector
 	{
-		bool operator()(const TWeakPtr<ISceneOutlinerTreeItem>& Item, FWorldPartitionActorDesc*& ActorDescPtrOut) const;
+		bool operator()(const TWeakPtr<ISceneOutlinerTreeItem>& Item, FWorldPartitionHandle& ActorHandleOut) const;
 	};
 
 	struct FFolderPathSelector
@@ -230,8 +236,11 @@ public:
 	/** Called by the editor to allow selection of unloaded actors */
 	void OnSelectUnloadedActors(const TArray<FGuid>& ActorGuids);
 	
-	/** Called when an actor desc is removed */
-	void OnActorDescRemoved(FWorldPartitionActorDesc* InActorDesc);
+	/** Called when an actor desc instance is removed */
+	void OnActorDescInstanceRemoved(FWorldPartitionActorDescInstance* InActorDescInstance);
+	
+	UE_DEPRECATED(5.4, "Use OnActorDescInstanceRemoved instead")
+	void OnActorDescRemoved(FWorldPartitionActorDesc* InActorDesc) {}
 
 	/** Called by engine when edit cut actors begins */
 	void OnEditCutActorsBegin();
@@ -339,7 +348,7 @@ protected:
 	/** Number of unloaded actors which have passed through all the filters */
 	uint32 FilteredUnloadedActorCount = 0;
 	/** List of unloaded actors which passed through the regular filters and may or may not have passed the search filter */
-	TSet<const FWorldPartitionActorDesc*> ApplicableUnloadedActors;
+	TSet<FWorldPartitionHandle> ApplicableUnloadedActors;
 	/** List of actors which passed the regular filters and may or may not have passed the search filter */
 	TSet<TWeakObjectPtr<AActor>> ApplicableActors;
 
diff --git a/Engine/Plugins/Editor/WorldPartitionHLODUtilities/Source/Private/WorldPartition/HLOD/Utilities/WorldPartitionHLODUtilities.cpp b/Engine/Plugins/Editor/WorldPartitionHLODUtilities/Source/Private/WorldPartition/HLOD/Utilities/WorldPartitionHLODUtilities.cpp
--- a/Engine/Plugins/Editor/WorldPartitionHLODUtilities/Source/Private/WorldPartition/HLOD/Utilities/WorldPartitionHLODUtilities.cpp
+++ b/Engine/Plugins/Editor/WorldPartitionHLODUtilities/Source/Private/WorldPartition/HLOD/Utilities/WorldPartitionHLODUtilities.cpp
@@ -1,12 +1,12 @@
 // Copyright Epic Games, Inc. All Rights Reserved.
 
 #include "WorldPartition/HLOD/Utilities/WorldPartitionHLODUtilities.h"
-#include "WorldPartition/WorldPartitionActorDescView.h"
 #include "WorldPartition/WorldPartitionStreamingGeneration.h"
 
 #if WITH_EDITOR
 
 #include "WorldPartition/WorldPartition.h"
+#include "WorldPartition/WorldPartitionActorDescInstance.h"
 #include "WorldPartition/ContentBundle/ContentBundleActivationScope.h"
 #include "WorldPartition/DataLayer/DataLayerInstance.h"
 #include "WorldPartition/HLOD/HLODActor.h"
@@ -79,13 +79,18 @@ static uint32 ComputeHLODHash(AWorldPartitionHLOD* InHLODActor, const TArray<UAc
 	return Ar.GetCrc();
 }
 
-void AddSourceActor(const FWorldPartitionActorDescView& ActorDescView, const IStreamingGenerationContext::FActorInstance& ActorInstance, const FName WorldPackageName, TMap<FGuid, FWorldPartitionRuntimeCellObjectMapping>& SourceActors)
+void AddSourceActor(const FStreamingGenerationActorDescView& ActorDescView, const FName WorldPackageName, TMap<FGuid, FWorldPartitionRuntimeCellObjectMapping>& SourceActors)
 {
 	const FName ActorPath = *ActorDescView.GetActorSoftPath().ToString();
 
-	const UActorDescContainer* ActorDescContainer = ActorDescView.GetActorDesc()->GetContainer();
+	const UActorDescContainerInstance* ContainerInstance = ActorDescView.GetContainerInstance();
+	check(ContainerInstance);
 
-	FGuid ActorInstanceGuid = ActorInstance.GetContainerID().GetActorGuid(ActorDescView.GetGuid());	
+	const FActorContainerID& ContainerID = ContainerInstance->GetContainerID();
+	const FTransform ContainerTransform = ContainerInstance->GetTransform();
+	const FName ContainerPackage = ContainerInstance->GetContainerPackage();
+
+	FGuid ActorInstanceGuid = ContainerID.GetActorGuid(ActorDescView.GetGuid());	
 	FWorldPartitionRuntimeCellObjectMapping& ActorMapping = SourceActors.FindOrAdd(ActorInstanceGuid);
 	if (!ActorMapping.ActorInstanceGuid.IsValid())
 	{
@@ -94,26 +99,26 @@ void AddSourceActor(const FWorldPartitionActorDescView& ActorDescView, const ISt
 			*ActorDescView.GetActorSoftPath().ToString(),
 			ActorDescView.GetBaseClass(),
 			ActorDescView.GetNativeClass(),
-			ActorInstance.GetContainerID(),
-			ActorInstance.GetTransform(),
-			ActorDescContainer->GetContainerPackage(),
+			ContainerID,
+			ContainerTransform,
+			ContainerPackage,
 			WorldPackageName,
 			ActorInstanceGuid,
 			false);
 
 		// Add its runtime references, recursively
-		const FStreamingGenerationActorDescViewMap* ActorDescViewMap = ActorInstance.ActorSetInstance->ContainerInstance->ActorDescViewMap;
+		const FStreamingGenerationActorDescViewMap& ActorDescViewMap = ActorDescView.GetActorDescViewMap();
 		for (const FGuid& ReferenceGuid : ActorDescView.GetReferences())
 		{
-			const FWorldPartitionActorDescView& RefActorDescView = ActorDescViewMap->FindByGuidChecked(ReferenceGuid);
-			AddSourceActor(RefActorDescView, ActorInstance, WorldPackageName, SourceActors);
+			const FStreamingGenerationActorDescView& RefActorDescView = ActorDescViewMap.FindByGuidChecked(ReferenceGuid);
+			AddSourceActor(RefActorDescView, WorldPackageName, SourceActors);
 		}
 
 		// Add its editor references
 		for (const FGuid& EditorReferenceGuid : ActorDescView.GetEditorReferences())
 		{
-			const FWorldPartitionActorDesc& ReferenceActorDesc = ActorDescContainer->GetActorDescChecked(EditorReferenceGuid);
-			FGuid EditorRefInstanceGuid = ActorInstance.GetContainerID().GetActorGuid(EditorReferenceGuid);
+			const FWorldPartitionActorDescInstance& ReferenceActorDesc = ContainerInstance->GetActorDescInstanceChecked(EditorReferenceGuid);
+			FGuid EditorRefInstanceGuid = ContainerID.GetActorGuid(EditorReferenceGuid);
 			FWorldPartitionRuntimeCellObjectMapping& EditorRefMapping = SourceActors.FindOrAdd(EditorRefInstanceGuid);
 			if (!EditorRefMapping.ActorInstanceGuid.IsValid())
 			{
@@ -122,9 +127,9 @@ void AddSourceActor(const FWorldPartitionActorDescView& ActorDescView, const ISt
 					*ReferenceActorDesc.GetActorSoftPath().ToString(),
 					ReferenceActorDesc.GetBaseClass(),
 					ReferenceActorDesc.GetNativeClass(),
-					ActorInstance.GetContainerID(),
-					ActorInstance.GetTransform(),
-					ActorDescContainer->GetContainerPackage(),
+					ContainerID,
+					ContainerTransform,
+					ContainerPackage,
 					WorldPackageName,
 					EditorRefInstanceGuid,
 					true
@@ -142,7 +147,7 @@ TArray<AWorldPartitionHLOD*> FWorldPartitionHLODUtilities::CreateHLODActors(FHLO
 
 	for (const IStreamingGenerationContext::FActorInstance& ActorInstance : InActors)
 	{
-		const FWorldPartitionActorDescView& ActorDescView = ActorInstance.GetActorDescView();
+		const FStreamingGenerationActorDescView& ActorDescView = ActorInstance.GetActorDescView();
 
 		if (ActorDescView.GetActorIsHLODRelevant())
 		{
@@ -155,7 +160,7 @@ TArray<AWorldPartitionHLOD*> FWorldPartitionHLODUtilities::CreateHLODActors(FHLO
 			if (UHLODLayer* HLODLayer = Cast<UHLODLayer>(ActorDescView.GetHLODLayer().TryLoad()))
 			{
 				TMap<FGuid, FWorldPartitionRuntimeCellObjectMapping>& SubActors = SourceActorsPerHLODLayer.FindOrAdd(HLODLayer);
-				AddSourceActor(ActorDescView, ActorInstance, WorldPackageName, SubActors);
+				AddSourceActor(ActorDescView, WorldPackageName, SubActors);
 			}
 		}
 	}
@@ -190,7 +195,7 @@ TArray<AWorldPartitionHLOD*> FWorldPartitionHLODUtilities::CreateHLODActors(FHLO
 		if (InCreationContext.HLODActorDescs.RemoveAndCopyValue(HLODActorName, HLODActorHandle))
 		{
 			InCreationContext.ActorReferences.Add(HLODActorHandle.ToReference());
-			HLODActor = CastChecked<AWorldPartitionHLOD>(HLODActorHandle->GetActor());
+			HLODActor = CastChecked<AWorldPartitionHLOD>(HLODActorHandle.GetActor());
 		}
 
 		bool bNewActor = HLODActor == nullptr;
diff --git a/Engine/Plugins/Experimental/PCG/Source/PCG/Private/PCGSubsystem.cpp b/Engine/Plugins/Experimental/PCG/Source/PCG/Private/PCGSubsystem.cpp
--- a/Engine/Plugins/Experimental/PCG/Source/PCG/Private/PCGSubsystem.cpp
+++ b/Engine/Plugins/Experimental/PCG/Source/PCG/Private/PCGSubsystem.cpp
@@ -28,6 +28,8 @@
 #include "PackageSourceControlHelper.h"
 #include "WorldPartition/WorldPartitionHelpers.h"
 #include "WorldPartition/ActorPartition/PartitionActorDesc.h"
+#include "WorldPartition/WorldPartitionActorDescInstance.h"
+#include "WorldPartition/WorldPartitionHandle.h"
 #include "ObjectTools.h"
 #else
 #include "Engine/Engine.h"
@@ -956,16 +958,16 @@ namespace PCGSubsystem
 					// TODO: Revisit after API review on the WP side, we shouldn't have to load here or get the actor desc directly
 					if (!PCGActor && bSaveActors)
 					{
-						const FWorldPartitionActorDesc* PCGActorDesc = nullptr;
-						auto FindFirst = [&CellCoord, &PCGActorDesc](const FWorldPartitionActorDesc* ActorDesc) {
-							FPartitionActorDesc* PartitionActorDesc = (FPartitionActorDesc*)ActorDesc;
+						const FWorldPartitionActorDescInstance* PCGActorDescInstance = nullptr;
+						auto FindFirst = [&CellCoord, &PCGActorDescInstance](const FWorldPartitionActorDescInstance* ActorDescInstance) {
+							const FPartitionActorDesc* PartitionActorDesc = (FPartitionActorDesc*)ActorDescInstance->GetActorDesc();
 
 							if (PartitionActorDesc &&
 								PartitionActorDesc->GridIndexX == CellCoord.X &&
 								PartitionActorDesc->GridIndexY == CellCoord.Y &&
 								PartitionActorDesc->GridIndexZ == CellCoord.Z)
 							{
-								PCGActorDesc = ActorDesc;
+								PCGActorDescInstance = ActorDescInstance;
 								return false;
 							}
 							else
@@ -974,13 +976,13 @@ namespace PCGSubsystem
 							}
 						};
 
-						FWorldPartitionHelpers::ForEachIntersectingActorDesc<APCGPartitionActor>(World->GetWorldPartition(), CellBounds, FindFirst);
+						FWorldPartitionHelpers::ForEachIntersectingActorDescInstance<APCGPartitionActor>(World->GetWorldPartition(), CellBounds, FindFirst);
 
-						check(!bCreateActor || PCGActorDesc);
-						if (PCGActorDesc)
+						check(!bCreateActor || PCGActorDescInstance);
+						if (PCGActorDescInstance)
 						{
-							ActorReferences->Add(FWorldPartitionReference(World->GetWorldPartition(), PCGActorDesc->GetGuid()));
-							PCGActor = Cast<APCGPartitionActor>(PCGActorDesc->GetActor());
+							ActorReferences->Add(FWorldPartitionReference(World->GetWorldPartition(), PCGActorDescInstance->GetGuid()));
+							PCGActor = Cast<APCGPartitionActor>(PCGActorDescInstance->GetActor());
 						}
 					}
 					// We still need to keep a reference on the PCG actor - note that newly created PCG actors will not have a reference here, but won't be unloaded
@@ -1007,16 +1009,16 @@ namespace PCGSubsystem
 					// We'll need to make sure actors in the bounds are loaded only if we need them.
 					if (bLoadCell)
 					{
-						auto WorldPartitionLoadActorsInBounds = [World, ActorReferences](const FWorldPartitionActorDesc* ActorDesc) {
-							check(ActorDesc);
-							ActorReferences->Add(FWorldPartitionReference(World->GetWorldPartition(), ActorDesc->GetGuid()));
+						auto WorldPartitionLoadActorsInBounds = [World, ActorReferences](const FWorldPartitionActorDescInstance* ActorDescInstance) {
+							check(ActorDescInstance);
+							ActorReferences->Add(FWorldPartitionReference(World->GetWorldPartition(), ActorDescInstance->GetGuid()));
 							// Load actor if not already loaded
-							ActorDesc->GetActor();
+							ActorDescInstance->GetActor();
 							return true;
 						};
 
 						auto LoadActorsTask = [World, IntersectedBounds, WorldPartitionLoadActorsInBounds]() {
-							FWorldPartitionHelpers::ForEachIntersectingActorDesc(World->GetWorldPartition(), IntersectedBounds, WorldPartitionLoadActorsInBounds);
+							FWorldPartitionHelpers::ForEachIntersectingActorDescInstance(World->GetWorldPartition(), IntersectedBounds, WorldPartitionLoadActorsInBounds);
 							return true;
 						};
 
@@ -1343,21 +1345,18 @@ void UPCGSubsystem::DeletePartitionActors(bool bOnlyDeleteUnused, bool bOnlyChil
 		return true;
 	};
 
-	auto GatherAndDestroyActors = [&PackagesToCleanup, &PackagesToDeleteFromSCC, World, &ActorsNotSafeToBeDeleted, &GatherAndDestroyLoadedActors](const FWorldPartitionActorDesc* ActorDesc) {
-		AActor* Actor = ActorDesc->GetActor();
-		if (!Actor)
-		{
-			Actor = ActorDesc->Load();
-		}
-
-		if (Actor)
+	auto GatherAndDestroyActors = [&PackagesToCleanup, &PackagesToDeleteFromSCC, World, &ActorsNotSafeToBeDeleted, &GatherAndDestroyLoadedActors](const FWorldPartitionActorDescInstance* ActorDescInstance) 
+	{
+		FWorldPartitionReference ActorReference(ActorDescInstance->GetContainerInstance(), ActorDescInstance->GetGuid());
+		
+		if (AActor* LoadedActor = ActorReference.GetActor())
 		{
-			GatherAndDestroyLoadedActors(Actor);
+			GatherAndDestroyLoadedActors(LoadedActor);
 		}
 		else // Couldn't load it
 		{
-			PackagesToDeleteFromSCC.Add(ActorDesc->GetActorPackage().ToString());
-			World->GetWorldPartition()->RemoveActor(ActorDesc->GetGuid());
+			PackagesToDeleteFromSCC.Add(ActorDescInstance->GetActorPackage().ToString());
+			World->GetWorldPartition()->RemoveActor(ActorDescInstance->GetGuid());
 		}
 
 		return true;
@@ -1400,7 +1399,7 @@ void UPCGSubsystem::DeletePartitionActors(bool bOnlyDeleteUnused, bool bOnlyChil
 	}
 	else
 	{
-		FWorldPartitionHelpers::ForEachActorDesc<APCGPartitionActor>(World->GetWorldPartition(), GatherAndDestroyActors);
+		FWorldPartitionHelpers::ForEachActorDescInstance<APCGPartitionActor>(World->GetWorldPartition(), GatherAndDestroyActors);
 
 		// Also cleanup the remaining actors that don't have descriptors, if we have a loaded level
 		if (ULevel* Level = World->GetCurrentLevel())
diff --git a/Engine/Plugins/Experimental/PCG/Source/PCG/Private/PCGWorldActor.cpp b/Engine/Plugins/Experimental/PCG/Source/PCG/Private/PCGWorldActor.cpp
--- a/Engine/Plugins/Experimental/PCG/Source/PCG/Private/PCGWorldActor.cpp
+++ b/Engine/Plugins/Experimental/PCG/Source/PCG/Private/PCGWorldActor.cpp
@@ -18,6 +18,7 @@
 #include "WorldPartition/WorldPartition.h"
 #include "WorldPartition/WorldPartitionHelpers.h"
 #include "WorldPartition/ActorPartition/PartitionActorDesc.h"
+#include "WorldPartition/WorldPartitionActorDescInstance.h"
 #endif
 
 #include UE_INLINE_GENERATED_CPP_BY_NAME(PCGWorldActor)
@@ -51,14 +52,14 @@ void APCGWorldActor::BeginCacheForCookedPlatformData(const ITargetPlatform* Targ
 
 		if (UWorldPartition* WorldPartition = World->GetWorldPartition())
 		{
-			FWorldPartitionHelpers::ForEachActorDesc<ALandscapeProxy>(WorldPartition, [WorldPartition, &ActorRefs](const FWorldPartitionActorDesc* ActorDesc)
+			FWorldPartitionHelpers::ForEachActorDescInstance<ALandscapeProxy>(WorldPartition, [WorldPartition, &ActorRefs](const FWorldPartitionActorDescInstance* ActorDescInstance)
 			{
-				check(ActorDesc);
+				check(ActorDescInstance);
 				// Create WP references only for actors that aren't currently loaded, otherwise we might end up unloading them
 				// if their actor desc ref count isn't setup properly
-				if (!ActorDesc->GetActor())
+				if (!ActorDescInstance->GetActor())
 				{
-					ActorRefs.Add(FWorldPartitionReference(WorldPartition, ActorDesc->GetGuid()));
+					ActorRefs.Add(FWorldPartitionReference(WorldPartition, ActorDescInstance->GetGuid()));
 				}
 				return true;
 			});
diff --git a/Engine/Plugins/Experimental/Water/Source/Runtime/Private/WaterSubsystem.cpp b/Engine/Plugins/Experimental/Water/Source/Runtime/Private/WaterSubsystem.cpp
--- a/Engine/Plugins/Experimental/Water/Source/Runtime/Private/WaterSubsystem.cpp
+++ b/Engine/Plugins/Experimental/Water/Source/Runtime/Private/WaterSubsystem.cpp
@@ -23,6 +23,7 @@
 
 #if WITH_EDITOR
 #include "WaterZoneActorDesc.h"
+#include "WorldPartition/WorldPartitionActorDescInstance.h"
 #include "WorldPartition/WorldPartitionHelpers.h"
 extern UNREALED_API UEditorEngine* GEditor;
 #else
@@ -552,15 +553,15 @@ TSoftObjectPtr<AWaterZone> UWaterSubsystem::FindWaterZone(const UWorld* World, c
 	{
 		if (UWorldPartition* WorldPartition = World->GetWorldPartition())
 		{
-			FWorldPartitionHelpers::ForEachActorDesc<AWaterZone>(WorldPartition, [Bounds, &ViableZones](const FWorldPartitionActorDesc* ActorDesc)
+			FWorldPartitionHelpers::ForEachActorDescInstance<AWaterZone>(WorldPartition, [Bounds, &ViableZones](const FWorldPartitionActorDescInstance* ActorDescInstance)
 			{
-				FWaterZoneActorDesc* WaterZoneActorDesc = (FWaterZoneActorDesc*)ActorDesc;
-				const FBox WaterZoneBounds = WaterZoneActorDesc->GetEditorBounds();
+				FWaterZoneActorDesc* WaterZoneActorDesc = (FWaterZoneActorDesc*)ActorDescInstance->GetActorDesc();
+				const FBox WaterZoneBounds = ActorDescInstance->GetEditorBounds();
 				const FBox2D WaterZoneBounds2D(FVector2D(WaterZoneBounds.Min), FVector2D(WaterZoneBounds.Max));
 
 				if (Bounds.Intersect(WaterZoneBounds2D))
 				{
-					ViableZones.Emplace(WaterZoneActorDesc->GetActorSoftPath(), WaterZoneActorDesc->GetOverlapPriority());
+					ViableZones.Emplace(ActorDescInstance->GetActorSoftPath(), WaterZoneActorDesc->GetOverlapPriority());
 				}
 
 				return true;
diff --git a/Engine/Source/Editor/DataLayerEditor/Private/DataLayer/DataLayerEditorSubsystem.cpp b/Engine/Source/Editor/DataLayerEditor/Private/DataLayer/DataLayerEditorSubsystem.cpp
--- a/Engine/Source/Editor/DataLayerEditor/Private/DataLayer/DataLayerEditorSubsystem.cpp
+++ b/Engine/Source/Editor/DataLayerEditor/Private/DataLayer/DataLayerEditorSubsystem.cpp
@@ -50,7 +50,6 @@
 #include "WorldPartition/DataLayer/DeprecatedDataLayerInstance.h"
 #include "WorldPartition/DataLayer/WorldDataLayers.h"
 #include "WorldPartition/DataLayer/DataLayerUtils.h"
-#include "WorldPartition/WorldPartitionActorDescViewProxy.h"
 #include "WorldPartition/WorldPartition.h"
 
 class SWidget;
@@ -1165,12 +1164,24 @@ bool UDataLayerEditorSubsystem::PassDataLayersFilter(UWorld* World, const FWorld
 
 	if (UDataLayerManager* DataLayerManager = UDataLayerManager::GetDataLayerManager(OwningWorld))
 	{
-		FWorldPartitionActorViewProxy ActorDescProxy(*ActorHandle);
-
+		// If Actor is loaded and dirty, use a newly resolved DataLayerInstanceNames array
+		auto GetLatestDataLayerInstanceNames = [DataLayerManager, ActorHandle]() -> TArray<FName>
+		{
+			if(AActor* Actor = ActorHandle.GetInstance()->GetActor(false); Actor && Actor->GetPackage()->IsDirty())
+			{
+				TUniquePtr<FWorldPartitionActorDesc> NewActorDesc = Actor->CreateActorDesc();
+				return FDataLayerUtils::ResolvedDataLayerInstanceNames(DataLayerManager, NewActorDesc.Get());
+			}
+			else
+			{
+				return ActorHandle.GetInstance()->GetDataLayerInstanceNames();
+			}
+		};
+				
 		if (IsRunningCookCommandlet())
 		{
 			// When running cook commandlet, dont allow loading of actors with runtime loaded data layers
-			for (const FName& DataLayerInstanceName : ActorDescProxy.GetDataLayerInstanceNames())
+			for (const FName& DataLayerInstanceName : GetLatestDataLayerInstanceNames())
 			{
 				const UDataLayerInstance* DataLayerInstance = DataLayerManager->GetDataLayerInstance(DataLayerInstanceName);
 				if (DataLayerInstance && DataLayerInstance->IsRuntime())
@@ -1182,7 +1193,7 @@ bool UDataLayerEditorSubsystem::PassDataLayersFilter(UWorld* World, const FWorld
 			return true;
 		}
 
-		return DataLayerManager->ResolveIsLoadedInEditor(ActorDescProxy.GetDataLayerInstanceNames());
+		return DataLayerManager->ResolveIsLoadedInEditor(GetLatestDataLayerInstanceNames());
 	}
 
 	return true;
diff --git a/Engine/Source/Editor/DataLayerEditor/Private/DataLayer/DataLayerHierarchy.cpp b/Engine/Source/Editor/DataLayerEditor/Private/DataLayer/DataLayerHierarchy.cpp
--- a/Engine/Source/Editor/DataLayerEditor/Private/DataLayer/DataLayerHierarchy.cpp
+++ b/Engine/Source/Editor/DataLayerEditor/Private/DataLayer/DataLayerHierarchy.cpp
@@ -22,7 +22,7 @@
 #include "UObject/ObjectPtr.h"
 #include "UObject/WeakObjectPtr.h"
 #include "WorldDataLayersTreeItem.h"
-#include "WorldPartition/ActorDescContainerCollection.h"
+#include "WorldPartition/ActorDescContainerInstanceCollection.h"
 #include "WorldPartition/DataLayer/DataLayerInstance.h"
 #include "WorldPartition/DataLayer/DataLayerInstanceWithAsset.h"
 #include "WorldPartition/DataLayer/DataLayerManager.h"
@@ -32,6 +32,7 @@
 #include "WorldPartition/WorldPartitionActorDesc.h"
 #include "WorldPartition/WorldPartitionHelpers.h"
 #include "WorldPartition/WorldPartitionSubsystem.h"
+#include "WorldPartition/WorldPartitionActorDescInstance.h"
 #include "EditorLoadedActorCache.h"
 
 TUniquePtr<FDataLayerHierarchy> FDataLayerHierarchy::Create(FDataLayerMode* Mode, const TWeakObjectPtr<UWorld>& World)
@@ -73,8 +74,8 @@ FDataLayerHierarchy::FDataLayerHierarchy(FDataLayerMode* Mode, const TWeakObject
 		
 		if (UWorldPartition* WorldPartition = World->GetWorldPartition())
 		{
-			WorldPartition->OnActorDescAddedEvent.AddRaw(this, &FDataLayerHierarchy::OnActorDescAdded);
-			WorldPartition->OnActorDescRemovedEvent.AddRaw(this, &FDataLayerHierarchy::OnActorDescRemoved);
+			WorldPartition->OnActorDescInstanceAddedEvent.AddRaw(this, &FDataLayerHierarchy::OnActorDescInstanceAdded);
+			WorldPartition->OnActorDescInstanceRemovedEvent.AddRaw(this, &FDataLayerHierarchy::OnActorDescInstanceRemoved);
 		}
 	}
 
@@ -109,8 +110,8 @@ FDataLayerHierarchy::~FDataLayerHierarchy()
 		
 		if (UWorldPartition* WorldPartition = RepresentingWorld->GetWorldPartition())
 		{
-			WorldPartition->OnActorDescAddedEvent.RemoveAll(this);
-			WorldPartition->OnActorDescRemovedEvent.RemoveAll(this);
+			WorldPartition->OnActorDescInstanceAddedEvent.RemoveAll(this);
+			WorldPartition->OnActorDescInstanceRemovedEvent.RemoveAll(this);
 		}
 	}
 
@@ -293,11 +294,11 @@ void FDataLayerHierarchy::CreateItems(TArray<FSceneOutlinerTreeItemPtr>& OutItem
 						const TSet<FGuid>& LoadedActors = LoadedActorCache.GetLoadedActorsForLevel(OuterLevel);
 
 						// Create an FDataLayerActorDescTreeItem for each unloaded actor of this WorldPartition
-						FWorldPartitionHelpers::ForEachActorDesc(WorldPartition, [this, IsDataLayerShown, DataLayerManager, CurrentLevel, &WorldToLevelDataLayerMap, &LoadedActors, &OutItems](const FWorldPartitionActorDesc* ActorDesc)
+						FWorldPartitionHelpers::ForEachActorDescInstance(WorldPartition, [this, IsDataLayerShown, DataLayerManager, CurrentLevel, &WorldToLevelDataLayerMap, &LoadedActors, &OutItems](const FWorldPartitionActorDescInstance* ActorDescInstance)
 						{
-							if (ActorDesc != nullptr && !LoadedActors.Contains(ActorDesc->GetGuid()) && FActorDescTreeItem::ShouldDisplayInOutliner(ActorDesc))
+							if (ActorDescInstance != nullptr && !LoadedActors.Contains(ActorDescInstance->GetGuid()) && FActorDescTreeItem::ShouldDisplayInOutliner(ActorDescInstance))
 							{
-								for (const FName& DataLayerInstanceName : ActorDesc->GetDataLayerInstanceNames())
+								for (const FName& DataLayerInstanceName : ActorDescInstance->GetDataLayerInstanceNames())
 								{
 									if (const UDataLayerInstance* DataLayerInstance = DataLayerManager->GetDataLayerInstance(DataLayerInstanceName))
 									{
@@ -310,7 +311,7 @@ void FDataLayerHierarchy::CreateItems(TArray<FSceneOutlinerTreeItemPtr>& OutItem
 												UDataLayerInstance** LevelDataLayerInstance = WorldToLevelDataLayerMap.Find(DataLayerInstance);
 												EffectiveDataLayerInstance = *LevelDataLayerInstance;
 											}
-											if (const FSceneOutlinerTreeItemPtr ActorDescItem = Mode->CreateItemFor<FDataLayerActorDescTreeItem>(FDataLayerActorDescTreeItemData(ActorDesc->GetGuid(), ActorDesc->GetContainer(), EffectiveDataLayerInstance)))
+											if (const FSceneOutlinerTreeItemPtr ActorDescItem = Mode->CreateItemFor<FDataLayerActorDescTreeItem>(FDataLayerActorDescTreeItemData(ActorDescInstance->GetGuid(), ActorDescInstance->GetContainerInstance(), EffectiveDataLayerInstance)))
 											{
 												OutItems.Add(ActorDescItem);
 											}
@@ -540,11 +541,11 @@ void FDataLayerHierarchy::OnLoadedActorAdded(AActor& InActor)
 		
 			if (UWorldPartition* WorldPartition = FWorldPartitionHelpers::GetWorldPartition(&InActor))
 			{
-				if (FWorldPartitionActorDesc* ActorDesc = WorldPartition->GetActorDesc(InActor.GetActorGuid()))
+				if (FWorldPartitionActorDescInstance* ActorDescInstance = WorldPartition->GetActorDescInstance(InActor.GetActorGuid()))
 				{
 					for (const UDataLayerInstance* DataLayerInstance : DataLayerInstances)
 					{
-						EventData.ItemIDs.Add(FDataLayerActorDescTreeItem::ComputeTreeItemID(ActorDesc->GetGuid(), ActorDesc->GetContainer(), DataLayerInstance));
+						EventData.ItemIDs.Add(FDataLayerActorDescTreeItem::ComputeTreeItemID(ActorDescInstance->GetGuid(), ActorDescInstance->GetContainerInstance(), DataLayerInstance));
 					}
 				}
 			}
@@ -571,9 +572,9 @@ void FDataLayerHierarchy::OnLoadedActorRemoved(AActor& InActor)
 			EventData.Items.Reserve(DataLayerInstances.Num());
 			for (const UDataLayerInstance* DataLayerInstance : DataLayerInstances)
 			{
-				if (FWorldPartitionActorDesc* ActorDesc = WorldPartition->GetActorDesc(InActor.GetActorGuid()); FActorDescTreeItem::ShouldDisplayInOutliner(ActorDesc))
+				if (FWorldPartitionActorDescInstance* ActorDescInstance = WorldPartition->GetActorDescInstance(InActor.GetActorGuid()); FActorDescTreeItem::ShouldDisplayInOutliner(ActorDescInstance))
 				{
-					EventData.Items.Add(Mode->CreateItemFor<FDataLayerActorDescTreeItem>(FDataLayerActorDescTreeItemData(ActorDesc->GetGuid(), ActorDesc->GetContainer(), const_cast<UDataLayerInstance*>(DataLayerInstance))));
+					EventData.Items.Add(Mode->CreateItemFor<FDataLayerActorDescTreeItem>(FDataLayerActorDescTreeItemData(ActorDescInstance->GetGuid(), ActorDescInstance->GetContainerInstance(), const_cast<UDataLayerInstance*>(DataLayerInstance))));
 				}
 			}
 			HierarchyChangedEvent.Broadcast(EventData);
@@ -581,15 +582,15 @@ void FDataLayerHierarchy::OnLoadedActorRemoved(AActor& InActor)
     }
 }
 
-void FDataLayerHierarchy::OnActorDescAdded(FWorldPartitionActorDesc* InActorDesc)
+void FDataLayerHierarchy::OnActorDescInstanceAdded(FWorldPartitionActorDescInstance* InActorDescInstance)
 {
-	if (!bShowUnloadedActors || !InActorDesc || InActorDesc->IsLoaded(true) || !FActorDescTreeItem::ShouldDisplayInOutliner(InActorDesc))
+	if (!bShowUnloadedActors || !InActorDescInstance || InActorDescInstance->IsLoaded(true) || !FActorDescTreeItem::ShouldDisplayInOutliner(InActorDescInstance))
 	{
 		return;
 	}
 
 	const UDataLayerManager* const DataLayerManager = UDataLayerManager::GetDataLayerManager(GetOwningWorld());
-	const TArray<FName>& DataLayerInstanceNames = InActorDesc->GetDataLayerInstanceNames();
+	const TArray<FName>& DataLayerInstanceNames = InActorDescInstance->GetDataLayerInstanceNames();
 
 	if (DataLayerManager && DataLayerInstanceNames.Num() > 0)
 	{
@@ -600,21 +601,21 @@ void FDataLayerHierarchy::OnActorDescAdded(FWorldPartitionActorDesc* InActorDesc
 		for (const FName& DataLayerInstanceName : DataLayerInstanceNames)
 		{
 			const UDataLayerInstance* const DataLayerInstance = DataLayerManager->GetDataLayerInstance(DataLayerInstanceName);
-			EventData.Items.Add(Mode->CreateItemFor<FDataLayerActorDescTreeItem>(FDataLayerActorDescTreeItemData(InActorDesc->GetGuid(), InActorDesc->GetContainer(), const_cast<UDataLayerInstance*>(DataLayerInstance))));
+			EventData.Items.Add(Mode->CreateItemFor<FDataLayerActorDescTreeItem>(FDataLayerActorDescTreeItemData(InActorDescInstance->GetGuid(), InActorDescInstance->GetContainerInstance(), const_cast<UDataLayerInstance*>(DataLayerInstance))));
 		}
 		HierarchyChangedEvent.Broadcast(EventData);
 	}
 }
 
-void FDataLayerHierarchy::OnActorDescRemoved(FWorldPartitionActorDesc* InActorDesc)
+void FDataLayerHierarchy::OnActorDescInstanceRemoved(FWorldPartitionActorDescInstance* InActorDescInstance)
 {
-	if (!bShowUnloadedActors || (InActorDesc == nullptr))
+	if (!bShowUnloadedActors || (InActorDescInstance == nullptr))
 	{
 		return;
 	}
 
 	const UDataLayerManager* const DataLayerManager = UDataLayerManager::GetDataLayerManager(GetOwningWorld());
-	const TArray<FName>& DataLayerInstanceNames = InActorDesc->GetDataLayerInstanceNames();
+	const TArray<FName>& DataLayerInstanceNames = InActorDescInstance->GetDataLayerInstanceNames();
 	
 	if (DataLayerManager && DataLayerInstanceNames.Num() > 0)
 	{
@@ -625,7 +626,7 @@ void FDataLayerHierarchy::OnActorDescRemoved(FWorldPartitionActorDesc* InActorDe
 		for (const FName& DataLayerInstanceName : DataLayerInstanceNames)
 		{
 			const UDataLayerInstance* const DataLayer = DataLayerManager->GetDataLayerInstance(DataLayerInstanceName);
-			EventData.ItemIDs.Add(FDataLayerActorDescTreeItem::ComputeTreeItemID(InActorDesc->GetGuid(), InActorDesc->GetContainer(), DataLayer));
+			EventData.ItemIDs.Add(FDataLayerActorDescTreeItem::ComputeTreeItemID(InActorDescInstance->GetGuid(), InActorDescInstance->GetContainerInstance(), DataLayer));
 		}
 		HierarchyChangedEvent.Broadcast(EventData);
 	}
diff --git a/Engine/Source/Editor/DataLayerEditor/Private/DataLayer/DataLayerHierarchy.h b/Engine/Source/Editor/DataLayerEditor/Private/DataLayer/DataLayerHierarchy.h
--- a/Engine/Source/Editor/DataLayerEditor/Private/DataLayer/DataLayerHierarchy.h
+++ b/Engine/Source/Editor/DataLayerEditor/Private/DataLayer/DataLayerHierarchy.h
@@ -15,6 +15,7 @@
 class AActor;
 class FDataLayerMode;
 class FWorldPartitionActorDesc;
+class FWorldPartitionActorDescInstance;
 class UDataLayerInstance;
 class ULevel;
 class UWorld;
@@ -55,8 +56,8 @@ private:
 	void OnLevelRemoved(ULevel* InLevel, UWorld* InWorld);
 	void OnLoadedActorAdded(AActor& InActor);
 	void OnLoadedActorRemoved(AActor& InActor);
-	void OnActorDescAdded(FWorldPartitionActorDesc* InActorDesc);
-	void OnActorDescRemoved(FWorldPartitionActorDesc* InActorDesc);
+	void OnActorDescInstanceAdded(FWorldPartitionActorDescInstance* InActorDescInstance);
+	void OnActorDescInstanceRemoved(FWorldPartitionActorDescInstance* InActorDescInstance);
 	void OnActorDataLayersChanged(const TWeakObjectPtr<AActor>& InActor);
 	void OnDataLayerChanged(const EDataLayerAction Action, const TWeakObjectPtr<const UDataLayerInstance>& ChangedDataLayer, const FName& ChangedProperty);
 	void FullRefreshEvent();
diff --git a/Engine/Source/Editor/DataLayerEditor/Private/DataLayer/DataLayerMode.cpp b/Engine/Source/Editor/DataLayerEditor/Private/DataLayer/DataLayerMode.cpp
--- a/Engine/Source/Editor/DataLayerEditor/Private/DataLayer/DataLayerMode.cpp
+++ b/Engine/Source/Editor/DataLayerEditor/Private/DataLayer/DataLayerMode.cpp
@@ -263,7 +263,7 @@ TSharedRef<FSceneOutlinerFilter> FDataLayerMode::CreateHideDataLayerActorsFilter
 
 TSharedRef<FSceneOutlinerFilter> FDataLayerMode::CreateHideUnloadedActorsFilter()
 {
-	return MakeShareable(new FActorDescFilter(FActorDescTreeItem::FFilterPredicate::CreateStatic([](const FWorldPartitionActorDesc* ActorDesc) { return true; }), FSceneOutlinerFilter::EDefaultBehaviour::Pass));
+	return MakeShareable(new FActorDescFilter(FActorDescTreeItem::FFilterPredicate::CreateStatic([](const FWorldPartitionActorDescInstance* ActorDescInstance) { return true; }), FSceneOutlinerFilter::EDefaultBehaviour::Pass));
 }
 
 int32 FDataLayerMode::GetTypeSortPriority(const ISceneOutlinerTreeItem& Item) const
diff --git a/Engine/Source/Editor/DataLayerEditor/Private/DataLayer/DataLayersActorDescTreeItem.h b/Engine/Source/Editor/DataLayerEditor/Private/DataLayer/DataLayersActorDescTreeItem.h
--- a/Engine/Source/Editor/DataLayerEditor/Private/DataLayer/DataLayersActorDescTreeItem.h
+++ b/Engine/Source/Editor/DataLayerEditor/Private/DataLayer/DataLayersActorDescTreeItem.h
@@ -7,7 +7,7 @@
 #include "SceneOutlinerFwd.h"
 #include "ActorDescTreeItem.h"
 #include "WorldPartition/DataLayer/DataLayerInstance.h"
-#include "WorldPartition/ActorDescContainerCollection.h"
+#include "WorldPartition/ActorDescContainerInstance.h"
 #include "WorldPartition/WorldPartition.h"
 #include "LevelInstance/LevelInstanceSubsystem.h"
 #include "Misc/ArchiveMD5.h"
@@ -22,14 +22,14 @@ class FWorldPartitionActorDesc;
 
 struct FDataLayerActorDescTreeItemData
 {
-	FDataLayerActorDescTreeItemData(const FGuid& InActorGuid, UActorDescContainer* InContainer, UDataLayerInstance* InDataLayer)
+	FDataLayerActorDescTreeItemData(const FGuid& InActorGuid, UActorDescContainerInstance* InContainerInstance, UDataLayerInstance* InDataLayer)
 		: ActorGuid(InActorGuid)
-		, Container(InContainer)
+		, ContainerInstance(InContainerInstance)
 		, DataLayer(InDataLayer)
 	{}
 
 	const FGuid& ActorGuid;
-	UActorDescContainer* const Container;
+	UActorDescContainerInstance* const ContainerInstance;
 	TWeakObjectPtr<UDataLayerInstance> DataLayer;
 };
 
@@ -39,20 +39,20 @@ struct FDataLayerActorDescTreeItemData
 struct FDataLayerActorDescTreeItem : public FActorDescTreeItem
 {
 public:
-	DECLARE_DELEGATE_RetVal_TwoParams(bool, FFilterPredicate, const FWorldPartitionActorDesc* ActorDesc, const UDataLayerInstance* DataLayer);
-	DECLARE_DELEGATE_RetVal_TwoParams(bool, FInteractivePredicate, const FWorldPartitionActorDesc* ActorDesc, const UDataLayerInstance* DataLayer);
+	DECLARE_DELEGATE_RetVal_TwoParams(bool, FFilterPredicate, const FWorldPartitionActorDescInstance* ActorDescInstance, const UDataLayerInstance* DataLayer);
+	DECLARE_DELEGATE_RetVal_TwoParams(bool, FInteractivePredicate, const FWorldPartitionActorDescInstance* ActorDescInstance, const UDataLayerInstance* DataLayer);
 
 	FDataLayerActorDescTreeItem(const FDataLayerActorDescTreeItemData& InData)
-		: FActorDescTreeItem(InData.ActorGuid, InData.Container)
+		: FActorDescTreeItem(InData.ActorGuid, InData.ContainerInstance)
 		, DataLayer(InData.DataLayer)
-		, IDDataLayerActorDesc(FDataLayerActorDescTreeItem::ComputeTreeItemID(InData.ActorGuid, InData.Container, InData.DataLayer.Get()))
+		, IDDataLayerActorDesc(FDataLayerActorDescTreeItem::ComputeTreeItemID(InData.ActorGuid, InData.ContainerInstance, InData.DataLayer.Get()))
 	{
 		if (ActorDescHandle.IsValid())
 		{
-			UActorDescContainer* Container = ActorDescHandle->GetContainer();
-			UWorld* OwningWorld = Container->GetWorldPartition()->GetWorld();
+			UActorDescContainerInstance* ContainerInstance = ActorDescHandle.GetContainerInstance();
+			UWorld* OwningWorld = ContainerInstance->GetWorldPartition()->GetWorld();
 			ULevelInstanceSubsystem* LevelInstanceSubsystem = UWorld::GetSubsystem<ULevelInstanceSubsystem>(OwningWorld);
-			ULevel* Level = Container->GetTypedOuter<UWorld>()->PersistentLevel;
+			ULevel* Level = ContainerInstance->GetTypedOuter<UWorld>()->PersistentLevel;
 			if (LevelInstanceSubsystem && Level && (Level != OwningWorld->GetCurrentLevel()))
 			{
 				DisplayString = LevelInstanceSubsystem->PrefixWithParentLevelInstanceActorLabels(DisplayString, Level);
@@ -62,9 +62,9 @@ public:
 
 	UDataLayerInstance* GetDataLayer() const { return DataLayer.Get(); }
 	
-	static FSceneOutlinerTreeItemID ComputeTreeItemID(FGuid InActorGuid, UActorDescContainer* InContainer, const UDataLayerInstance* InDataLayer)
+		static FSceneOutlinerTreeItemID ComputeTreeItemID(FGuid InActorGuid, UActorDescContainerInstance* InContainerInstance, const UDataLayerInstance* InDataLayer)
 	{
-		FObjectKey ContainerKey(InContainer);
+		FObjectKey ContainerKey(InContainerInstance);
 		FObjectKey DataLayerInstanceKey(InDataLayer);
 
 		FArchiveMD5 Ar;
@@ -73,28 +73,14 @@ public:
 		return FSceneOutlinerTreeItemID(Ar.GetGuidFromHash());
 	}
 
-	static TArray<AActor*> GetParentActors(UActorDescContainer* InContainer)
-	{
-		if (InContainer)
-		{
-			UWorld* OwningWorld = InContainer->GetWorldPartition()->GetWorld();
-			if (ULevelInstanceSubsystem* LevelInstanceSubsystem = UWorld::GetSubsystem<ULevelInstanceSubsystem>(OwningWorld))
-			{
-				ULevel* Level = InContainer->GetTypedOuter<UWorld>()->PersistentLevel;
-				return LevelInstanceSubsystem->GetParentLevelInstanceActors(Level);
-			}
-		}
-		return TArray<AActor*>();
-	}
-
 	bool Filter(FFilterPredicate Pred) const
 	{
-		return Pred.Execute(ActorDescHandle.Get(), DataLayer.Get());
+		return Pred.Execute(ActorDescHandle.GetInstance(), DataLayer.Get());
 	}
 
 	bool GetInteractiveState(FInteractivePredicate Pred) const
 	{
-		return Pred.Execute(ActorDescHandle.Get(), DataLayer.Get());
+		return Pred.Execute(ActorDescHandle.GetInstance(), DataLayer.Get());
 	}
 
 	/* Begin ISceneOutlinerTreeItem Implementation */
diff --git a/Engine/Source/Editor/DataLayerEditor/Private/DataLayer/SDataLayerBrowser.cpp b/Engine/Source/Editor/DataLayerEditor/Private/DataLayer/SDataLayerBrowser.cpp
--- a/Engine/Source/Editor/DataLayerEditor/Private/DataLayer/SDataLayerBrowser.cpp
+++ b/Engine/Source/Editor/DataLayerEditor/Private/DataLayer/SDataLayerBrowser.cpp
@@ -99,9 +99,9 @@ void SDataLayerBrowser::Construct(const FArguments& InArgs)
 		}
 		else if (const FDataLayerActorDescTreeItem* ActorDescItem = Item.CastTo<FDataLayerActorDescTreeItem>())
 		{
-			if (const FWorldPartitionActorDesc* ActorDesc = ActorDescItem->ActorDescHandle.Get())
+			if (const FWorldPartitionActorDescInstance* ActorDescInstance = ActorDescItem->ActorDescHandle.GetInstance())
 			{
-				return ActorDesc->GetActorName().ToString();
+				return ActorDescInstance->GetActorName().ToString();
 			}
 		}
 		return FString();
diff --git a/Engine/Source/Editor/PropertyEditor/Private/UserInterface/PropertyEditor/SPropertyEditorAsset.cpp b/Engine/Source/Editor/PropertyEditor/Private/UserInterface/PropertyEditor/SPropertyEditorAsset.cpp
--- a/Engine/Source/Editor/PropertyEditor/Private/UserInterface/PropertyEditor/SPropertyEditorAsset.cpp
+++ b/Engine/Source/Editor/PropertyEditor/Private/UserInterface/PropertyEditor/SPropertyEditorAsset.cpp
@@ -31,6 +31,7 @@
 #include "Framework/Notifications/NotificationManager.h"
 #include "Widgets/Notifications/SNotificationList.h"
 #include "WorldPartition/WorldPartition.h"
+#include "WorldPartition/WorldPartitionActorDescInstance.h"
 #include "FileHelpers.h"
 #include "Presentation/PropertyEditor/PropertyEditor.h"
 #include "AssetThumbnail.h"
@@ -1311,9 +1312,9 @@ void SPropertyEditorAsset::OnBrowse()
 				{
 					if (UWorld* World = Cast<UWorld>(MapObject); World && World->IsPartitionedWorld())
 					{
-						if (const FWorldPartitionActorDesc* ActorDesc = World->GetWorldPartition()->GetActorDescByPath(Value.ObjectPath))
+						if (const FWorldPartitionActorDescInstance* ActorDescInstance = World->GetWorldPartition()->GetActorDescInstanceByPath(Value.ObjectPath))
 						{
-							World->GetWorldPartition()->PinActors({ ActorDesc->GetGuid() });
+							World->GetWorldPartition()->PinActors({ ActorDescInstance->GetGuid() });
 							GetValue(Value);
 						}
 					}
diff --git a/Engine/Source/Editor/SceneOutliner/Private/ActorBrowsingMode.cpp b/Engine/Source/Editor/SceneOutliner/Private/ActorBrowsingMode.cpp
--- a/Engine/Source/Editor/SceneOutliner/Private/ActorBrowsingMode.cpp
+++ b/Engine/Source/Editor/SceneOutliner/Private/ActorBrowsingMode.cpp
@@ -41,6 +41,8 @@
 #include "WorldPartition/WorldPartitionSubsystem.h"
 #include "WorldPartition/DataLayer/DataLayerManager.h"
 #include "WorldPartition/IWorldPartitionEditorModule.h"
+#include "WorldPartition/ActorDescContainerInstance.h"
+#include "WorldPartition/WorldPartitionActorDescInstance.h"
 #include "Subsystems/ActorEditorContextSubsystem.h"
 #include "ISourceControlModule.h"
 #include "ISourceControlProvider.h"
@@ -142,7 +144,7 @@ FActorBrowsingMode::FActorBrowsingMode(SSceneOutliner* InSceneOutliner, TWeakObj
 	FilterInfoMap.Add(TEXT("ShowOnlyCurrentDataLayers"), OnlyCurrentDataLayersInfo);
 
 	// Add a filter for unloaded actors to properly reflect the bShowOnlyActorsInCurrentDataLayers flag.
-	SceneOutliner->AddFilter(MakeShared<FActorDescFilter>(FActorDescTreeItem::FFilterPredicate::CreateLambda([this](const FWorldPartitionActorDesc* ActorDesc)
+	SceneOutliner->AddFilter(MakeShared<FActorDescFilter>(FActorDescTreeItem::FFilterPredicate::CreateLambda([this](const FWorldPartitionActorDescInstance* ActorDescInstance)
 		{
 			FActorBrowsingModeConfig* Settings = GetMutableConfig();
 			if (Settings && Settings->bShowOnlyActorsInCurrentDataLayers)
@@ -153,7 +155,7 @@ FActorBrowsingMode::FActorBrowsingMode(SSceneOutliner* InSceneOutliner, TWeakObj
 					return true;
 				}
 				
-				for(const UDataLayerInstance* const DataLayerInstance : DataLayerManager->GetDataLayerInstances(ActorDesc->GetDataLayerInstanceNames()))
+				for(const UDataLayerInstance* const DataLayerInstance : DataLayerManager->GetDataLayerInstances(ActorDescInstance->GetDataLayerInstanceNames()))
 				{
 					if (DataLayerInstance->IsInActorEditorContext())
 					{
@@ -178,7 +180,7 @@ FActorBrowsingMode::FActorBrowsingMode(SSceneOutliner* InSceneOutliner, TWeakObj
 	FilterInfoMap.Add(TEXT("ShowOnlyCurrentContentBundle"), OnlyCurrentContentBundleInfo);
 
 	// Add a filter for unloaded actors to properly reflect the bShowOnlyActorsInCurrentContentBundle flag.
-	SceneOutliner->AddFilter(MakeShared<FActorDescFilter>(FActorDescTreeItem::FFilterPredicate::CreateLambda([this](const FWorldPartitionActorDesc* ActorDesc)
+	SceneOutliner->AddFilter(MakeShared<FActorDescFilter>(FActorDescTreeItem::FFilterPredicate::CreateLambda([this](const FWorldPartitionActorDescInstance* ActorDescInstance)
 		{
 			FActorBrowsingModeConfig* Settings = GetMutableConfig();
 			if (Settings && Settings->bShowOnlyActorsInCurrentContentBundle)
@@ -188,7 +190,7 @@ FActorBrowsingMode::FActorBrowsingMode(SSceneOutliner* InSceneOutliner, TWeakObj
 					return true;
 				}
 
-				return ActorDesc->GetContentBundleGuid().IsValid() && WorldPartitionEditorModule->IsEditingContentBundle(ActorDesc->GetContentBundleGuid());
+				return ActorDescInstance->GetContentBundleGuid().IsValid() && WorldPartitionEditorModule->IsEditingContentBundle(ActorDescInstance->GetContentBundleGuid());
 			}
 			return true;
 		}), FSceneOutlinerFilter::EDefaultBehaviour::Pass));
@@ -284,7 +286,7 @@ FActorBrowsingMode::~FActorBrowsingMode()
 	{
 		if (UWorldPartition* const WorldPartition = RepresentingWorld->GetWorldPartition())
 		{
-			WorldPartition->OnActorDescRemovedEvent.RemoveAll(this);
+			WorldPartition->OnActorDescInstanceRemovedEvent.RemoveAll(this);
 		}
 	}
 	FSceneOutlinerDelegates::Get().OnComponentsUpdated.RemoveAll(this);
@@ -317,7 +319,7 @@ void FActorBrowsingMode::Rebuild()
 	{
 		if (UWorldPartition* const WorldPartition = RepresentingWorld->GetWorldPartition())
 		{
-			WorldPartition->OnActorDescRemovedEvent.RemoveAll(this);
+			WorldPartition->OnActorDescInstanceRemovedEvent.RemoveAll(this);
 		}
 	}
 
@@ -334,7 +336,7 @@ void FActorBrowsingMode::Rebuild()
 	if (bRepresentingWorldPartitionedWorld)
 	{
 		UWorldPartition* const WorldPartition = RepresentingWorld->GetWorldPartition();
-		WorldPartition->OnActorDescRemovedEvent.AddRaw(this, &FActorBrowsingMode::OnActorDescRemoved);
+		WorldPartition->OnActorDescInstanceRemovedEvent.AddRaw(this, &FActorBrowsingMode::OnActorDescInstanceRemoved);
 	}
 }
 
@@ -633,7 +635,7 @@ TSharedRef<FSceneOutlinerFilter> FActorBrowsingMode::CreateHideComponentsFilter(
 TSharedRef<FSceneOutlinerFilter> FActorBrowsingMode::CreateHideUnloadedActorsFilter()
 {
 	return MakeShareable(new FActorDescFilter(FActorDescTreeItem::FFilterPredicate::CreateStatic(
-		[](const FWorldPartitionActorDesc* ActorDesc) { return false; }), FSceneOutlinerFilter::EDefaultBehaviour::Pass));
+		[](const FWorldPartitionActorDescInstance* ActorDescInstance) { return false; }), FSceneOutlinerFilter::EDefaultBehaviour::Pass));
 }
 
 TSharedRef<FSceneOutlinerFilter> FActorBrowsingMode::CreateHideEmptyFoldersFilter()
@@ -973,9 +975,9 @@ void FActorBrowsingMode::OnSelectUnloadedActors(const TArray<FGuid>& ActorGuids)
 		ItemsToSelect.Reserve(ActorGuids.Num());
 		for (const FGuid& ActorGuid : ActorGuids)
 		{
-			if (FWorldPartitionActorDesc* ActorDesc = WorldPartition->GetActorDesc(ActorGuid))
+			if (FWorldPartitionActorDescInstance* ActorDescInstance = WorldPartition->GetActorDescInstance(ActorGuid))
 			{
-				if (FSceneOutlinerTreeItemPtr ItemPtr = SceneOutliner->GetTreeItem(FActorDescTreeItem::ComputeTreeItemID(ActorDesc->GetGuid(), ActorDesc->GetContainer())))
+				if (FSceneOutlinerTreeItemPtr ItemPtr = SceneOutliner->GetTreeItem(FActorDescTreeItem::ComputeTreeItemID(ActorDescInstance->GetGuid(), ActorDescInstance->GetContainerInstance())))
 				{
 					ItemsToSelect.Add(ItemPtr);
 				}
@@ -995,9 +997,9 @@ void FActorBrowsingMode::OnSelectUnloadedActors(const TArray<FGuid>& ActorGuids)
 	}
 }
 
-void FActorBrowsingMode::OnActorDescRemoved(FWorldPartitionActorDesc* InActorDesc)
+void FActorBrowsingMode::OnActorDescInstanceRemoved(FWorldPartitionActorDescInstance* InActorDescInstance)
 {
-	ApplicableUnloadedActors.Remove(InActorDesc);
+	ApplicableUnloadedActors.Remove(InActorDescInstance);
 }
 
 void FActorBrowsingMode::OnItemSelectionChanged(FSceneOutlinerTreeItemPtr TreeItem, ESelectInfo::Type SelectionType, const FSceneOutlinerItemSelection& Selection)
@@ -1175,7 +1177,7 @@ void FActorBrowsingMode::OnItemPassesFilters(const ISceneOutlinerTreeItem& Item)
 	}
 	else if (const FActorDescTreeItem* const ActorDescItem = Item.CastTo<FActorDescTreeItem>(); ActorDescItem && ActorDescItem->IsValid())
 	{
-		ApplicableUnloadedActors.Add(ActorDescItem->ActorDescHandle.Get());
+		ApplicableUnloadedActors.Add(ActorDescItem->ActorDescHandle);
 	}
 }
 
@@ -1302,9 +1304,9 @@ bool FActorBrowsingMode::HasErrors() const
 		{
 			bool bHasErrors = 0;
 
-			WorldPartition->ForEachActorDescContainer([&bHasErrors](UActorDescContainer* ActorDescContainer)
+			WorldPartition->ForEachActorDescContainer([&bHasErrors](UActorDescContainerInstance* ActorDescContainerInstance)
 			{
-				if (ActorDescContainer->HasInvalidActors())
+				if (ActorDescContainerInstance->GetContainer()->HasInvalidActors())
 				{
 					bHasErrors = true;
 				}
@@ -1332,13 +1334,13 @@ void FActorBrowsingMode::RepairErrors() const
 			ISourceControlProvider& SourceControlProvider = SourceControlModule.GetProvider();
 
 			TArray<FAssetData> InvalidActorAssets;
-			WorldPartition->ForEachActorDescContainer([&InvalidActorAssets](UActorDescContainer* ActorDescContainer)
+			WorldPartition->ForEachActorDescContainer([&InvalidActorAssets](UActorDescContainerInstance* ActorDescContainerInstance)
 			{
-				for (const FAssetData& InvalidActor : ActorDescContainer->GetInvalidActors())
+				for (const FAssetData& InvalidActor : ActorDescContainerInstance->GetContainer()->GetInvalidActors())
 				{
 					InvalidActorAssets.Add(InvalidActor);
 				}
-				ActorDescContainer->ClearInvalidActors();
+				ActorDescContainerInstance->GetContainer()->ClearInvalidActors();
 			});
 
 			TArray<FString> ActorFilesToDelete;
@@ -1416,12 +1418,12 @@ void FActorBrowsingMode::SynchronizeSelectedActorDescs()
 	if (UWorldPartitionSubsystem* WorldPartitionSubsystem = UWorld::GetSubsystem<UWorldPartitionSubsystem>(RepresentingWorld.Get()))
 	{
 		const FSceneOutlinerItemSelection Selection = SceneOutliner->GetSelection();
-		TArray<FWorldPartitionActorDesc*> SelectedActorDescs = Selection.GetData<FWorldPartitionActorDesc*>(SceneOutliner::FActorDescSelector());
+		TArray<FWorldPartitionHandle> SelectedActorHandles = Selection.GetData<FWorldPartitionHandle>(SceneOutliner::FActorHandleSelector());
 
-		WorldPartitionSubsystem->SelectedActorDescs.Empty();
-		for (FWorldPartitionActorDesc* SelectedActorDesc : SelectedActorDescs)
+		WorldPartitionSubsystem->SelectedActorHandles.Empty();
+		for (const FWorldPartitionHandle& ActorHandle : SelectedActorHandles)
 		{
-			WorldPartitionSubsystem->SelectedActorDescs.Add(SelectedActorDesc);
+			WorldPartitionSubsystem->SelectedActorHandles.Add(ActorHandle);
 		}
 	}
 }
@@ -1562,11 +1564,11 @@ namespace ActorBrowsingModeUtils
 		
 		for (const TTuple<UWorld*, TSet<FName>>& Pair : UnloadedActorsFolderPaths)
 		{
-			FActorFolders::ForEachActorDescInFolders(*Pair.Key, Pair.Value, [&List](const FWorldPartitionActorDesc* ActorDesc)
+			FActorFolders::ForEachActorDescInstanceInFolders(*Pair.Key, Pair.Value, [&List](const FWorldPartitionActorDescInstance* ActorDescInstance)
 			{
-				if (!ActorDesc->IsLoaded())
+				if (!ActorDescInstance->IsLoaded())
 				{
-					List.Add(ActorDesc->GetGuid());
+					List.Add(ActorDescInstance->GetGuid());
 				}
 				return true;
 			});
@@ -1646,7 +1648,7 @@ void FActorBrowsingMode::PinItems(const TArray<FSceneOutlinerTreeItemPtr>& InIte
 		{
 			if (FWorldPartitionHandle ActorHandle(WorldPartition, ActorGuid); ActorHandle.IsValid())
 			{
-				if (AActor* PinnedActor = ActorHandle->GetActor())
+				if (AActor* PinnedActor = ActorHandle.GetInstance()->GetActor())
 				{
 					GEditor->SelectActor(PinnedActor, /*bInSelected=*/true, /*bNotify=*/false);
 					LastPinnedActor = PinnedActor;
@@ -1850,10 +1852,10 @@ bool FActorBrowsingMode::CompareItemWithClassName(SceneOutliner::FilterBarType I
 	}
 	else if (const FActorDescTreeItem* ActorDescItem = InItem.CastTo<FActorDescTreeItem>())
 	{
-		if (const FWorldPartitionActorDesc* ActorDesc = ActorDescItem->ActorDescHandle.Get())
+		if (const FWorldPartitionActorDescInstance* ActorDescInstance = ActorDescItem->ActorDescHandle.GetInstance())
 		{
 			// For Unloaded Actors, grab the native class 
-			FTopLevelAssetPath ClassPath = ActorDesc->GetNativeClass();
+			FTopLevelAssetPath ClassPath = ActorDescInstance->GetNativeClass();
 			return AssetClassPaths.Contains(ClassPath);
 		}
 	}
