Object Mixer:

-Fix: Content Browser Add menu - "Blueprint Class" and "Object Mixer Filter" share the same Tooltip

diff --git a/Engine/Plugins/Editor/ObjectMixer/ObjectMixer/Source/ObjectMixer/Private/ObjectMixerFilterFactory.cpp b/Engine/Plugins/Editor/ObjectMixer/ObjectMixer/Source/ObjectMixer/Private/ObjectMixerFilterFactory.cpp
--- a/Engine/Plugins/Editor/ObjectMixer/ObjectMixer/Source/ObjectMixer/Private/ObjectMixerFilterFactory.cpp
+++ b/Engine/Plugins/Editor/ObjectMixer/ObjectMixer/Source/ObjectMixer/Private/ObjectMixerFilterFactory.cpp
@@ -46,4 +46,14 @@ uint32 UObjectMixerBlueprintFilterFactory::GetMenuCategories() const
 	return EAssetTypeCategories::Misc;
 }
 
+FText UObjectMixerBlueprintFilterFactory::GetToolTip() const
+{
+	return LOCTEXT("ObjectMixerFilterFactoryTooltip", "Creates a new Object Mixer Filter that can be used to define rules that govern what objects, columns and properties appear in Object Mixer.");
+}
+
+FString UObjectMixerBlueprintFilterFactory::GetToolTipDocumentationExcerpt() const
+{
+	return TEXT("UObjectMixerBlueprintFilter");
+}
+
 #undef LOCTEXT_NAMESPACE
diff --git a/Engine/Plugins/Editor/ObjectMixer/ObjectMixer/Source/ObjectMixer/Public/ObjectMixerFilterFactory.h b/Engine/Plugins/Editor/ObjectMixer/ObjectMixer/Source/ObjectMixer/Public/ObjectMixerFilterFactory.h
--- a/Engine/Plugins/Editor/ObjectMixer/ObjectMixer/Source/ObjectMixer/Public/ObjectMixerFilterFactory.h
+++ b/Engine/Plugins/Editor/ObjectMixer/ObjectMixer/Source/ObjectMixer/Public/ObjectMixerFilterFactory.h
@@ -22,6 +22,8 @@ public:
 	virtual UObject* FactoryCreateNew(UClass* Class, UObject* InParent, FName Name, EObjectFlags Flags, UObject* Context, FFeedbackContext* Warn) override;
 	virtual FString GetDefaultNewAssetName() const override;
 	virtual uint32 GetMenuCategories() const override;
+	virtual FText GetToolTip() const override;
+	virtual FString GetToolTipDocumentationExcerpt() const override;
 	//~ Begin UFactory Interface
 	
 };
