Object Mixer:

-New: Add option to turn off hybrid rows as project setting
-New: Add option to select actor vs component based on project setting when clicking hybrid row
-Remove deprecated struct FActorDescSelector
-Remove bExpandTreeViewItemsByDefault from Project Settings since it's not currently functional. Will implement at another time.
-Add project settings to Show Options menu in the main widget for convenience

diff --git a/Engine/Plugins/Editor/ObjectMixer/ObjectMixer/Source/ObjectMixer/ObjectMixerEditor.build.cs b/Engine/Plugins/Editor/ObjectMixer/ObjectMixer/Source/ObjectMixer/ObjectMixerEditor.build.cs
--- a/Engine/Plugins/Editor/ObjectMixer/ObjectMixer/Source/ObjectMixer/ObjectMixerEditor.build.cs
+++ b/Engine/Plugins/Editor/ObjectMixer/ObjectMixer/Source/ObjectMixer/ObjectMixerEditor.build.cs
@@ -35,6 +35,7 @@ public class ObjectMixerEditor : ModuleRules
 				"AssetTools",
 				"CoreUObject",
 				"ContentBrowser",
+				"DeveloperSettings",
 				"Engine",
 				"EditorConfig",
 				"EditorStyle",
diff --git a/Engine/Plugins/Editor/ObjectMixer/ObjectMixer/Source/ObjectMixer/Private/ObjectMixerEditorModule.cpp b/Engine/Plugins/Editor/ObjectMixer/ObjectMixer/Source/ObjectMixer/Private/ObjectMixerEditorModule.cpp
--- a/Engine/Plugins/Editor/ObjectMixer/ObjectMixer/Source/ObjectMixer/Private/ObjectMixerEditorModule.cpp
+++ b/Engine/Plugins/Editor/ObjectMixer/ObjectMixer/Source/ObjectMixer/Private/ObjectMixerEditorModule.cpp
@@ -95,6 +95,14 @@ void FObjectMixerEditorModule::Teardown()
 
 	FLevelEditorSequencerIntegration::Get().GetOnSequencersChanged().RemoveAll(this);
 
+	if (FModuleManager::Get().IsModuleLoaded(TEXT("ObjectMixer")))
+	{
+		if (UObjectMixerEditorSettings* Settings = GetMutableDefault<UObjectMixerEditorSettings>())
+		{
+			Settings->OnSettingChanged().RemoveAll(this);
+		}
+	}
+
 	for (FDelegateHandle& Delegate : DelegateHandles)
 	{
 		Delegate.Reset();
@@ -382,9 +390,18 @@ void FObjectMixerEditorModule::BindDelegates()
 	}));
 	
 	DelegateHandles.Add(FLevelEditorSequencerIntegration::Get().GetOnSequencersChanged().AddLambda(
-		[this]
+	[this]
+	{
+		RegenerateListWidget();
+	}));
+
+	DelegateHandles.Add(GetMutableDefault<UObjectMixerEditorSettings>()->OnSettingChanged().AddLambda(
+		[this](UObject*, FPropertyChangedEvent& Event)
 		{
-			RegenerateListWidget();
+			if (UObjectMixerEditorSettings::DoesPropertyChangeRequireListRebuild(Event))
+			{
+				RequestRebuildList();
+			}
 		}));
 }
 
diff --git a/Engine/Plugins/Editor/ObjectMixer/ObjectMixer/Source/ObjectMixer/Private/Views/List/Modes/ObjectMixerOutlinerHierarchy.cpp b/Engine/Plugins/Editor/ObjectMixer/ObjectMixer/Source/ObjectMixer/Private/Views/List/Modes/ObjectMixerOutlinerHierarchy.cpp
--- a/Engine/Plugins/Editor/ObjectMixer/ObjectMixer/Source/ObjectMixer/Private/Views/List/Modes/ObjectMixerOutlinerHierarchy.cpp
+++ b/Engine/Plugins/Editor/ObjectMixer/ObjectMixer/Source/ObjectMixer/Private/Views/List/Modes/ObjectMixerOutlinerHierarchy.cpp
@@ -19,6 +19,7 @@
 #include "EditorActorFolders.h"
 #include "ISceneOutlinerMode.h"
 #include "LevelTreeItem.h"
+#include "ObjectMixerEditorSettings.h"
 #include "WorldTreeItem.h"
 #include "Engine/Engine.h"
 #include "Engine/Level.h"
@@ -606,8 +607,9 @@ TArray<FSceneOutlinerTreeItemPtr> FObjectMixerOutlinerHierarchy::ConditionallyCr
 		if (const FSceneOutlinerTreeItemPtr ActorItem =
 			Mode->CreateItemFor<FObjectMixerEditorListRowActor>(
 				FObjectMixerEditorListRowActor(Actor, GetCastedMode()->GetSceneOutliner())))
-		{	
-			if (ComponentRows.Num() == 1) // Create hybrid row
+		{
+			const bool bIsHybridModeEnabled = GetDefault<UObjectMixerEditorSettings>()->IsHybridRowModeEnabled();
+			if (bIsHybridModeEnabled && ComponentRows.Num() == 1) // Create hybrid row
 			{
 				if (FObjectMixerEditorListRowActor* AsActorRow = FObjectMixerUtils::AsActorRow(ActorItem))
 				{
diff --git a/Engine/Plugins/Editor/ObjectMixer/ObjectMixer/Source/ObjectMixer/Private/Views/List/Modes/ObjectMixerOutlinerMode.cpp b/Engine/Plugins/Editor/ObjectMixer/ObjectMixer/Source/ObjectMixer/Private/Views/List/Modes/ObjectMixerOutlinerMode.cpp
--- a/Engine/Plugins/Editor/ObjectMixer/ObjectMixer/Source/ObjectMixer/Private/Views/List/Modes/ObjectMixerOutlinerMode.cpp
+++ b/Engine/Plugins/Editor/ObjectMixer/ObjectMixer/Source/ObjectMixer/Private/Views/List/Modes/ObjectMixerOutlinerMode.cpp
@@ -119,7 +119,21 @@ namespace ObjectMixerOutliner
 					return true;
 				}
 			}
+			else if (FObjectMixerEditorListRowActor* ActorItem = ItemPtr->CastTo<FObjectMixerEditorListRowActor>())
+			{
+				if (ActorItem->IsValid())
+				{
+					// Get the component from a hybrid row if we're in HybridComponent mode
+					if (GetDefault<UObjectMixerEditorSettings>()->HybridRowPolicy == EObjectMixerHybridMode::HybridComponent &&
+						ActorItem->RowData.GetHybridComponent())
+					{
+						DataOut = ActorItem->RowData.GetHybridComponent();
+						return true;
+					}
+				}
+			}
 		}
+		
 		return false;
 	}
 
@@ -143,10 +157,17 @@ namespace ObjectMixerOutliner
 	{
 		if (TSharedPtr<ISceneOutlinerTreeItem> ItemPtr = Item.Pin())
 		{
-			if (FActorTreeItem* ActorItem = ItemPtr->CastTo<FActorTreeItem>())
+			if (FObjectMixerEditorListRowActor* ActorItem = ItemPtr->CastTo<FObjectMixerEditorListRowActor>())
 			{
 				if (ActorItem->IsValid())
 				{
+					// Skip hybrid rows here if in HybridComponent mode, we'll get the components from them in FComponentSelector
+					if (GetDefault<UObjectMixerEditorSettings>()->HybridRowPolicy == EObjectMixerHybridMode::HybridComponent &&
+						ActorItem->RowData.GetIsHybridRow())
+					{
+						return false;
+					}
+					
 					AActor* Actor = ActorItem->Actor.Get();
 					if (Actor)
 					{
@@ -1340,7 +1361,19 @@ void FObjectMixerOutlinerMode::CreateViewContent(FMenuBuilder& MenuBuilder)
 	}
 	MenuBuilder.EndSection();
 
-	MenuBuilder.BeginSection("AssetThumbnails", LOCTEXT("ShowWorldHeading", "World"));
+	MenuBuilder.BeginSection("ProjectSettings", LOCTEXT("ShowProjectSettingsHeading", "Project Settings"));
+	{
+		FDetailsViewArgs DetailsViewArgs;
+		DetailsViewArgs.bAllowSearch = false;
+		DetailsViewArgs.bShowOptions = false;
+		FPropertyEditorModule & PropertyEditorModule = FModuleManager::LoadModuleChecked<FPropertyEditorModule>("PropertyEditor");
+		const TSharedRef<IDetailsView> DetailsView = PropertyEditorModule.CreateDetailView(DetailsViewArgs);
+		DetailsView->SetObject(GetMutableDefault<UObjectMixerEditorSettings>());
+		MenuBuilder.AddWidget(DetailsView, FText::GetEmpty(), true, false);
+	}
+	MenuBuilder.EndSection();
+
+	MenuBuilder.BeginSection("WorldPicker", LOCTEXT("ShowWorldHeading", "World"));
 	{
 		MenuBuilder.AddSubMenu(
 			LOCTEXT("ChooseWorldSubMenu", "Choose World"),
diff --git a/Engine/Plugins/Editor/ObjectMixer/ObjectMixer/Source/ObjectMixer/Public/ObjectMixerEditorSettings.h b/Engine/Plugins/Editor/ObjectMixer/ObjectMixer/Source/ObjectMixer/Public/ObjectMixerEditorSettings.h
--- a/Engine/Plugins/Editor/ObjectMixer/ObjectMixer/Source/ObjectMixer/Public/ObjectMixerEditorSettings.h
+++ b/Engine/Plugins/Editor/ObjectMixer/ObjectMixer/Source/ObjectMixer/Public/ObjectMixerEditorSettings.h
@@ -2,16 +2,34 @@
 
 #pragma once
 
+#include "Engine/DeveloperSettings.h"
+
 #include "ObjectMixerEditorSettings.generated.h"
 
+UENUM(BlueprintType)
+enum class EObjectMixerHybridMode : uint8
+{
+	HybridActor UMETA(DisplayName="Hybrid Rows (Select Actor)", ToolTip="Makes hybrid rows and selects the actor in the Scene Outliner when an actor with a single matching component is selected."),
+	HybridComponent UMETA(DisplayName="Hybrid Rows (Select Component)", ToolTip="Makes hybrid rows and selects the component in the Scene Outliner when an actor with a single matching component is selected."),
+	HybridNone UMETA(DisplayName="Do not make hybrid rows", ToolTip="Does not make hybrid rows when an actor with a single matching component is selected. Actor and component rows will be split."),
+};
+
 UCLASS(config = ObjectMixer)
-class OBJECTMIXEREDITOR_API UObjectMixerEditorSettings : public UObject
+class OBJECTMIXEREDITOR_API UObjectMixerEditorSettings : public UDeveloperSettings
 {
 	GENERATED_BODY()
+	
 public:
 	
-	UObjectMixerEditorSettings(const FObjectInitializer& ObjectInitializer)
-	{}
+	static bool DoesPropertyChangeRequireListRebuild(const FPropertyChangedEvent& Event)
+	{
+		return Event.GetPropertyName() == GET_MEMBER_NAME_CHECKED(UObjectMixerEditorSettings, HybridRowPolicy);
+	}
+
+	bool IsHybridRowModeEnabled() const
+	{
+		return HybridRowPolicy != EObjectMixerHybridMode::HybridNone;
+	}
 	
 	/**
 	 * If enabled, clicking an item in the mixer list will also select the item in the Scene Outliner.
@@ -22,8 +40,10 @@ public:
 	bool bSyncSelection = true;
 
 	/**
-	 * If false, a new object will be created every time the filter object is accessed.
+	 * Determines how rows with a single matching component are displayed. By default, these rows are condensed into a single row.
+	 * Selecting that row will select the actor in the scene outliner, but not the component. You can choose to select the component instead,
+	 * or choose to not condense the rows and leave them separated.
 	 */
 	UPROPERTY(Config, EditAnywhere, Category="Object Mixer")
-	bool bExpandTreeViewItemsByDefault = true;
+	EObjectMixerHybridMode HybridRowPolicy = EObjectMixerHybridMode::HybridActor;
 };
diff --git a/Engine/Plugins/Editor/ObjectMixer/ObjectMixer/Source/ObjectMixer/Public/Views/List/Modes/ObjectMixerOutlinerMode.h b/Engine/Plugins/Editor/ObjectMixer/ObjectMixer/Source/ObjectMixer/Public/Views/List/Modes/ObjectMixerOutlinerMode.h
--- a/Engine/Plugins/Editor/ObjectMixer/ObjectMixer/Source/ObjectMixer/Public/Views/List/Modes/ObjectMixerOutlinerMode.h
+++ b/Engine/Plugins/Editor/ObjectMixer/ObjectMixer/Source/ObjectMixer/Public/Views/List/Modes/ObjectMixerOutlinerMode.h
@@ -38,11 +38,6 @@ namespace ObjectMixerOutliner
 	{
 		bool operator()(const TWeakPtr<ISceneOutlinerTreeItem>& Item, AActor*& ActorPtrOut) const;
 	};
-		
-	struct UE_DEPRECATED(5.4, "Use FActorHandleSelector instead") FActorDescSelector
-	{
-		bool operator()(const TWeakPtr<ISceneOutlinerTreeItem>& Item, FWorldPartitionActorDesc*& ActorDescPtrOut) const { return false; }
-	};
 
 	/** Functor which can be used to get actor descriptors from a selection  */
 	struct FActorHandleSelector
