Object Mixer:

-Fix: "Show Only in any Current Data Layers" filter does not function in Light Mixer

diff --git a/Engine/Plugins/Editor/ObjectMixer/ObjectMixer/Source/ObjectMixer/Private/Views/List/Modes/ObjectMixerOutlinerMode.cpp b/Engine/Plugins/Editor/ObjectMixer/ObjectMixer/Source/ObjectMixer/Private/Views/List/Modes/ObjectMixerOutlinerMode.cpp
--- a/Engine/Plugins/Editor/ObjectMixer/ObjectMixer/Source/ObjectMixer/Private/Views/List/Modes/ObjectMixerOutlinerMode.cpp
+++ b/Engine/Plugins/Editor/ObjectMixer/ObjectMixer/Source/ObjectMixer/Private/Views/List/Modes/ObjectMixerOutlinerMode.cpp
@@ -61,6 +61,7 @@
 #include "WorldPartition/WorldPartition.h"
 #include "WorldPartition/WorldPartitionSubsystem.h"
 #include "WorldPartition/WorldPartitionActorDescInstance.h"
+#include "WorldPartition/DataLayer/DataLayerManager.h"
 
 static int32 GSceneOutlinerAutoRepresentingWorldNetModeForObjectMixer = NM_Client;
 static FAutoConsoleVariableRef CVarAutoRepresentingWorldNetMode(
@@ -1418,7 +1419,15 @@ TSharedRef<FSceneOutlinerFilter> FObjectMixerOutlinerMode::CreateIsInCurrentData
 {
 	return MakeShareable(new FActorFilter(FActorTreeItem::FFilterPredicate::CreateStatic([](const AActor* InActor)
 		{
-			return true;
+			for (const UDataLayerInstance* DataLayerInstance : InActor->GetDataLayerInstances())
+				{
+					if (DataLayerInstance->IsInActorEditorContext())
+					{
+						return true;
+					}
+				}
+
+				return false;
 		}), FSceneOutlinerFilter::EDefaultBehaviour::Pass));
 }
 
