Object Mixer:

-Add ability to create sequencer anim keys from property widgets

diff --git a/Engine/Plugins/Editor/ObjectMixer/ObjectMixer/Source/ObjectMixer/Private/ObjectMixerEditorModule.cpp b/Engine/Plugins/Editor/ObjectMixer/ObjectMixer/Source/ObjectMixer/Private/ObjectMixerEditorModule.cpp
--- a/Engine/Plugins/Editor/ObjectMixer/ObjectMixer/Source/ObjectMixer/Private/ObjectMixerEditorModule.cpp
+++ b/Engine/Plugins/Editor/ObjectMixer/ObjectMixer/Source/ObjectMixer/Private/ObjectMixerEditorModule.cpp
@@ -152,6 +152,11 @@ TSharedPtr<SWidget> FObjectMixerEditorModule::MakeObjectMixerDialog(
 	return ObjectMixerDialog;
 }
 
+TArray<TWeakPtr<ISequencer>> FObjectMixerEditorModule::GetSequencers() const
+{
+	return FLevelEditorSequencerIntegration::Get().GetSequencers();
+}
+
 TSharedPtr<SDockTab> FObjectMixerEditorModule::FindNomadTab()
 {
 	if (!DockTab.IsValid())
@@ -388,9 +393,8 @@ void FObjectMixerEditorModule::BindDelegates()
 	{
 		RequestRebuildList();
 	}));
-	
-	DelegateHandles.Add(FLevelEditorSequencerIntegration::Get().GetOnSequencersChanged().AddLambda(
-	[this]
+
+	DelegateHandles.Add(FLevelEditorSequencerIntegration::Get().GetOnSequencersChanged().AddLambda([this]
 	{
 		RegenerateListWidget();
 	}));
diff --git a/Engine/Plugins/Editor/ObjectMixer/ObjectMixer/Source/ObjectMixer/Private/Views/List/SObjectMixerEditorList.cpp b/Engine/Plugins/Editor/ObjectMixer/ObjectMixer/Source/ObjectMixer/Private/Views/List/SObjectMixerEditorList.cpp
--- a/Engine/Plugins/Editor/ObjectMixer/ObjectMixer/Source/ObjectMixer/Private/Views/List/SObjectMixerEditorList.cpp
+++ b/Engine/Plugins/Editor/ObjectMixer/ObjectMixer/Source/ObjectMixer/Private/Views/List/SObjectMixerEditorList.cpp
@@ -21,12 +21,11 @@
 #include "EditorActorFolders.h"
 #include "IPlacementModeModule.h"
 #include "LevelEditorSequencerIntegration.h"
-#include "LevelSequence.h"
 #include "SceneOutlinerModule.h"
 #include "SceneOutlinerTextInfoColumn.h"
 #include "SPositiveActionButton.h"
+
 #include "Algo/Find.h"
-#include "Algo/RemoveIf.h"
 #include "Framework/Application/SlateApplication.h"
 #include "Framework/MultiBox/MultiBoxBuilder.h"
 #include "LevelInstance/LevelInstanceInterface.h"
@@ -34,12 +33,10 @@
 #include "Widgets/Input/SCheckBox.h"
 #include "Widgets/Layout/SWrapBox.h"
 #include "Widgets/Views/STreeView.h"
-#include "WorldPartition/ActorDescContainer.h"
-#include "WorldPartition/WorldPartitionActorDesc.h"
+#include "WorldPartition/WorldPartitionActorDescInstance.h"
 #include "WorldPartition/ContentBundle/ContentBundleDescriptor.h"
 #include "WorldPartition/ContentBundle/ContentBundleEngineSubsystem.h"
 #include "WorldPartition/DataLayer/DataLayerManager.h"
-#include "WorldPartition/WorldPartitionActorDescInstance.h"
 
 #define LOCTEXT_NAMESPACE "SObjectMixerEditorList"
 
@@ -100,7 +97,6 @@ void SObjectMixerEditorList::Construct(const FArguments& InArgs, TSharedRef<FObj
 
 SObjectMixerEditorList::~SObjectMixerEditorList()
 {	
-	FLevelEditorSequencerIntegration::Get().GetOnSequencersChanged().RemoveAll(this);
 	HeaderRowContextMenuWidget.Reset();
 }
 
@@ -1188,7 +1184,7 @@ void SObjectMixerEditorList::SetupColumns(FSceneOutlinerInitializationOptions& O
 		);
 		
 		// Sequencer Columns
-		if (FLevelEditorSequencerIntegration::Get().GetSequencers().Num() > 0)
+		if (FObjectMixerEditorModule::Get().GetSequencers().Num() > 0)
 		{
 			const FText SpawnableDisplayText = LOCTEXT("SpawnableColumnName", "Spawnable");
 			const FText SequencerInfoDisplayText = LOCTEXT("SequencerColumnName", "Sequencer");
diff --git a/Engine/Plugins/Editor/ObjectMixer/ObjectMixer/Source/ObjectMixer/Private/Views/Widgets/SInlinePropertyCellWidget.cpp b/Engine/Plugins/Editor/ObjectMixer/ObjectMixer/Source/ObjectMixer/Private/Views/Widgets/SInlinePropertyCellWidget.cpp
--- a/Engine/Plugins/Editor/ObjectMixer/ObjectMixer/Source/ObjectMixer/Private/Views/Widgets/SInlinePropertyCellWidget.cpp
+++ b/Engine/Plugins/Editor/ObjectMixer/ObjectMixer/Source/ObjectMixer/Private/Views/Widgets/SInlinePropertyCellWidget.cpp
@@ -2,17 +2,44 @@
 
 #include "Views/Widgets/SInlinePropertyCellWidget.h"
 
+#include "ObjectMixerEditorModule.h"
 #include "Views/List/ObjectMixerEditorListRowData.h"
 #include "Views/List/ObjectMixerUtils.h"
 
+#include "ISequencer.h"
+#include "ISinglePropertyView.h"
+#include "KeyParams.h"
+#include "KeyPropertyParams.h"
+#include "SSimpleButton.h"
 #include "Framework/Application/MenuStack.h"
 #include "Framework/Application/SlateApplication.h"
 #include "HAL/PlatformApplicationMisc.h"
-#include "ISinglePropertyView.h"
 #include "Layout/WidgetPath.h"
 #include "Modules/ModuleManager.h"
 
 #define LOCTEXT_NAMESPACE "ObjectMixerEditor"
+
+namespace UE::ObjectMixer
+{
+	void CreateKey(UObject* InObject, TSharedPtr<IPropertyHandle> InPropertyHandle)
+	{
+		if (!IsValid(InObject) || !InPropertyHandle.IsValid() || !InPropertyHandle->IsValidHandle())
+		{
+			return;
+		}
+
+		for (TWeakPtr<ISequencer> Sequencer : FObjectMixerEditorModule::Get().GetSequencers())
+		{
+			if (!Sequencer.IsValid())
+			{
+				continue;
+			}
+
+			const FKeyPropertyParams KeyPropertyParams({InObject}, *InPropertyHandle, ESequencerKeyMode::ManualKeyForced);
+			Sequencer.Pin()->KeyProperty(KeyPropertyParams);
+		}
+	}
+}
 	
 void SInlinePropertyCellWidget::Construct(
 	const FArguments& InArgs, const FName InColumnName, const TSharedRef<ISceneOutlinerTreeItem> RowPtr)
@@ -37,7 +64,9 @@ void SInlinePropertyCellWidget::Construct(
 
 	if (SinglePropertyView)
 	{
-		if (const TSharedPtr<IPropertyHandle> Handle = SinglePropertyView->GetPropertyHandle())
+		const TSharedPtr<IPropertyHandle> Handle = SinglePropertyView->GetPropertyHandle();
+		
+		if (Handle && Handle->IsValidHandle())
 		{
 			FObjectMixerUtils::GetRowData(RowPtr)->PropertyNamesToHandles.Add(InColumnName, Handle);
 				
@@ -50,13 +79,29 @@ void SInlinePropertyCellWidget::Construct(
 
 			ChildSlot
 			[
-				SNew(SBox)
-				.Visibility(EVisibility::Visible)
+				SNew(SHorizontalBox)
+
+				+ SHorizontalBox::Slot()
 				.HAlign(HAlign_Fill)
 				.VAlign(VAlign_Center)
 				[
 					SinglePropertyView.ToSharedRef()
 				]
+
+				+ SHorizontalBox::Slot()
+				.AutoWidth()
+				.VAlign(VAlign_Center)
+				[
+					SNew(SSimpleButton)
+					.ToolTipText(LOCTEXT("CreateKeyToolTip", "Add a keyframe for this property."))
+					.OnClicked_Lambda([ObjectRef, Handle]()
+					{
+						UE::ObjectMixer::CreateKey(ObjectRef, Handle);
+						return FReply::Handled();
+					})
+					.Visibility(FObjectMixerEditorModule::Get().GetSequencers().Num() > 0 ? EVisibility::Visible : EVisibility::Collapsed)
+					.Icon(FAppStyle::Get().GetBrush("Sequencer.AddKey.Details"))
+				]
 			];
 		}
 	}
diff --git a/Engine/Plugins/Editor/ObjectMixer/ObjectMixer/Source/ObjectMixer/Public/ObjectMixerEditorModule.h b/Engine/Plugins/Editor/ObjectMixer/ObjectMixer/Source/ObjectMixer/Public/ObjectMixerEditorModule.h
--- a/Engine/Plugins/Editor/ObjectMixer/ObjectMixer/Source/ObjectMixer/Public/ObjectMixerEditorModule.h
+++ b/Engine/Plugins/Editor/ObjectMixer/ObjectMixer/Source/ObjectMixer/Public/ObjectMixerEditorModule.h
@@ -7,7 +7,7 @@
 #include "Widgets/Docking/SDockTab.h"
 
 class FObjectMixerEditorList;
-class FObjectMixerEditorList;
+class ISequencer;
 
 class OBJECTMIXEREDITOR_API FObjectMixerEditorModule : public IModuleInterface
 {
@@ -32,6 +32,9 @@ public:
 	virtual TSharedPtr<SWidget> MakeObjectMixerDialog(
 		TSubclassOf<UObjectMixerObjectFilter> InDefaultFilterClass = nullptr);
 
+	/** Get a list of sequencers that are currently active in the editor. */
+	virtual TArray<TWeakPtr<ISequencer>> GetSequencers() const;
+
 	/**
 	 * Tries to find the nomad tab assigned to this instance of Object Mixer.
 	 * If DockTab is not set, will try to find the tab using GetTabSpawnerId().
diff --git a/Engine/Plugins/Editor/ObjectMixer/ObjectMixer/Source/ObjectMixer/Public/Views/Widgets/SInlinePropertyCellWidget.h b/Engine/Plugins/Editor/ObjectMixer/ObjectMixer/Source/ObjectMixer/Public/Views/Widgets/SInlinePropertyCellWidget.h
--- a/Engine/Plugins/Editor/ObjectMixer/ObjectMixer/Source/ObjectMixer/Public/Views/Widgets/SInlinePropertyCellWidget.h
+++ b/Engine/Plugins/Editor/ObjectMixer/ObjectMixer/Source/ObjectMixer/Public/Views/Widgets/SInlinePropertyCellWidget.h
@@ -6,6 +6,8 @@
 #include "Templates/SharedPointer.h"
 #include "UObject/NameTypes.h"
 
+class IPropertyHandle;
+
 struct FObjectMixerEditorListRowData;
 
 class SInlinePropertyCellWidget : public SCompoundWidget
