Add a scoped transaction when modifying the solo state.

diff --git a/Engine/Plugins/Editor/ObjectMixer/ObjectMixer/Source/ObjectMixer/Private/Views/List/Modes/OutlinerColumns/ObjectMixerOutlinerColumns.cpp b/Engine/Plugins/Editor/ObjectMixer/ObjectMixer/Source/ObjectMixer/Private/Views/List/Modes/OutlinerColumns/ObjectMixerOutlinerColumns.cpp
--- a/Engine/Plugins/Editor/ObjectMixer/ObjectMixer/Source/ObjectMixer/Private/Views/List/Modes/OutlinerColumns/ObjectMixerOutlinerColumns.cpp
+++ b/Engine/Plugins/Editor/ObjectMixer/ObjectMixer/Source/ObjectMixer/Private/Views/List/Modes/OutlinerColumns/ObjectMixerOutlinerColumns.cpp
@@ -164,8 +164,12 @@ void FObjectMixerOutlinerSoloColumn::OnClickSoloIcon(const FSceneOutlinerTreeIte
 			{
 				SetSoloPerRowRecursively(RowPtr, bNewSolo);
 			}
-								
-			ListView->EvaluateAndSetEditorVisibilityPerRow();
+
+			{
+				FScopedTransaction SoloObject(LOCTEXT("ChangeSoloState", "Solo Object"));
+				ListView->EvaluateAndSetEditorVisibilityPerRow();
+			}
+
 		}
 	}
 }
