[ColorGradingEditor] Display color grading target objects in an ObjectMixer-based hierarchical view

diff --git a/Engine/Plugins/Editor/ColorGrading/ColorGrading.uplugin b/Engine/Plugins/Editor/ColorGrading/ColorGrading.uplugin
--- a/Engine/Plugins/Editor/ColorGrading/ColorGrading.uplugin
+++ b/Engine/Plugins/Editor/ColorGrading/ColorGrading.uplugin
@@ -25,6 +25,10 @@
 		{
 			"Name": "ColorCorrectRegions",
 			"Enabled": true
+		},
+		{
+			"Name": "ObjectMixer",
+			"Enabled": true
 		}
 	]
 }
diff --git a/Engine/Plugins/Editor/ColorGrading/Source/ColorGradingEditor/ColorGradingEditor.build.cs b/Engine/Plugins/Editor/ColorGrading/Source/ColorGradingEditor/ColorGradingEditor.build.cs
--- a/Engine/Plugins/Editor/ColorGrading/Source/ColorGradingEditor/ColorGradingEditor.build.cs
+++ b/Engine/Plugins/Editor/ColorGrading/Source/ColorGradingEditor/ColorGradingEditor.build.cs
@@ -7,6 +7,14 @@ public class ColorGradingEditor : ModuleRules
 {
 	public ColorGradingEditor(ReadOnlyTargetRules ROTargetRules) : base(ROTargetRules)
 	{
+		PublicDependencyModuleNames.AddRange(
+			new string[]
+			{
+				"SceneOutliner",
+				"ObjectMixerEditor"
+			}
+		);
+
 		PrivateDependencyModuleNames.AddRange(
 			new string[] {
 				"ApplicationCore",
@@ -17,8 +25,8 @@ public class ColorGradingEditor : ModuleRules
 				"DetailCustomizations",
 				"EditorStyle",
 				"Engine",
-				"LevelEditor",
 				"InputCore",
+				"LevelEditor",
 				"PropertyEditor",
 				"PropertyPath",
 				"Slate",
diff --git a/Engine/Plugins/Editor/ColorGrading/Source/ColorGradingEditor/Private/ColorGradingEditorModule.cpp b/Engine/Plugins/Editor/ColorGrading/Source/ColorGradingEditor/Private/ColorGradingEditorModule.cpp
--- a/Engine/Plugins/Editor/ColorGrading/Source/ColorGradingEditor/Private/ColorGradingEditorModule.cpp
+++ b/Engine/Plugins/Editor/ColorGrading/Source/ColorGradingEditor/Private/ColorGradingEditorModule.cpp
@@ -4,11 +4,13 @@
 
 #include "ColorGradingCommands.h"
 #include "ColorGradingEditorDataModel.h"
+#include "ColorGradingMixerObjectFilterRegistry.h"
 #include "DataModelGenerators/ColorGradingDataModelGenerator_PostProcessVolume.h"
 #include "DataModelGenerators/ColorGradingDataModelGenerator_ColorCorrectRegion.h"
 #include "SColorGradingPanel.h"
 
 #include "ColorCorrectRegion.h"
+#include "ColorCorrectWindow.h"
 #include "Engine/PostProcessVolume.h"
 #include "Framework/Docking/LayoutExtender.h"
 #include "LevelEditor.h"
@@ -30,6 +32,13 @@ void FColorGradingEditorModule::StartupModule()
 	FColorGradingListItem::RegisterColorGradingListItemGenerator<FColorGradingListItemGenerator_PostProcessVolume>();
 	FColorGradingListItem::RegisterColorGradingListItemGenerator<FColorGradingListItemGenerator_ColorCorrectRegion>();
 
+	FColorGradingMixerObjectFilterRegistry::RegisterActorClassToPlace(APostProcessVolume::StaticClass());
+	FColorGradingMixerObjectFilterRegistry::RegisterActorClassToPlace(AColorCorrectionRegion::StaticClass());
+	FColorGradingMixerObjectFilterRegistry::RegisterActorClassToPlace(AColorCorrectionWindow::StaticClass());
+
+	FColorGradingMixerObjectFilterRegistry::RegisterObjectClassToFilter(APostProcessVolume::StaticClass());
+	FColorGradingMixerObjectFilterRegistry::RegisterObjectClassToFilter(AColorCorrectRegion::StaticClass());
+
 	FColorGradingCommands::Register();
 
 	FCoreDelegates::OnFEngineLoopInitComplete.AddRaw(this, &FColorGradingEditorModule::OnFEngineLoopInitComplete);
diff --git a/Engine/Plugins/Editor/ColorGrading/Source/ColorGradingEditor/Private/ColorGradingMixerObjectFilter.cpp b/Engine/Plugins/Editor/ColorGrading/Source/ColorGradingEditor/Private/ColorGradingMixerObjectFilter.cpp
new file mode 100644
--- /dev/null
+++ b/Engine/Plugins/Editor/ColorGrading/Source/ColorGradingEditor/Private/ColorGradingMixerObjectFilter.cpp
@@ -0,0 +1,14 @@
+// Copyright Epic Games, Inc. All Rights Reserved.
+
+#include "ColorGradingMixerObjectFilter.h"
+#include "ColorGradingMixerObjectFilterRegistry.h"
+
+TSet<UClass*> UColorGradingMixerObjectFilter::GetObjectClassesToFilter() const
+{
+	return FColorGradingMixerObjectFilterRegistry::GetObjectClassesToFilter();
+}
+
+TSet<TSubclassOf<AActor>> UColorGradingMixerObjectFilter::GetObjectClassesToPlace() const
+{
+	return FColorGradingMixerObjectFilterRegistry::GetActorClassesToPlace();
+}
\ No newline at end of file
diff --git a/Engine/Plugins/Editor/ColorGrading/Source/ColorGradingEditor/Private/ColorGradingMixerObjectFilter.h b/Engine/Plugins/Editor/ColorGrading/Source/ColorGradingEditor/Private/ColorGradingMixerObjectFilter.h
new file mode 100644
--- /dev/null
+++ b/Engine/Plugins/Editor/ColorGrading/Source/ColorGradingEditor/Private/ColorGradingMixerObjectFilter.h
@@ -0,0 +1,65 @@
+// Copyright Epic Games, Inc. All Rights Reserved.
+
+#pragma once
+
+#include "ObjectFilter/ObjectMixerEditorObjectFilter.h"
+
+#include "ColorGradingMixerObjectFilter.generated.h"
+
+UCLASS(BlueprintType, EditInlineNew)
+class COLORGRADINGEDITOR_API UColorGradingMixerObjectFilter : public UObjectMixerObjectFilter
+{
+	GENERATED_BODY()
+public:
+
+	virtual TSet<UClass*> GetObjectClassesToFilter() const override;
+	virtual TSet<TSubclassOf<AActor>> GetObjectClassesToPlace() const override;
+
+	virtual bool GetShowTransientObjects() const override
+	{
+		return true;
+	}
+
+	virtual TSet<FName> GetColumnsToShowByDefault() const override
+	{
+		return {};
+	}
+
+	virtual TSet<FName> GetColumnsToExclude() const
+	{
+		return {};
+	}
+
+	virtual TSet<FName> GetForceAddedColumns() const override
+	{
+		return {};
+	}
+
+	virtual bool ShouldIncludeUnsupportedProperties() const override
+	{
+		return false;
+	}
+
+	virtual bool ShouldAllowHybridRows() const override
+	{
+		// Disabled because both actors and components can have color grading settings, so they must be displayed on separate rows
+		// to disambiguate which one is selected for editing. If hybrid mode is enabled, they would be folded into the same row,
+		// making selection of the component impossible from the Color Grading panel.
+		return false;
+	}
+
+	virtual bool ShouldAllowColumnCustomizationByUser() const override
+	{
+		return false;
+	}
+
+	virtual EObjectMixerInheritanceInclusionOptions GetObjectMixerPropertyInheritanceInclusionOptions() const override
+	{
+		return EObjectMixerInheritanceInclusionOptions::IncludeAllParentsAndChildren;
+	}
+
+	virtual EObjectMixerInheritanceInclusionOptions GetObjectMixerPlacementClassInclusionOptions() const override
+	{
+		return EObjectMixerInheritanceInclusionOptions::None;
+	}
+};
diff --git a/Engine/Plugins/Editor/ColorGrading/Source/ColorGradingEditor/Private/SColorGradingColorWheelPanel.cpp b/Engine/Plugins/Editor/ColorGrading/Source/ColorGradingEditor/Private/SColorGradingColorWheelPanel.cpp
--- a/Engine/Plugins/Editor/ColorGrading/Source/ColorGradingEditor/Private/SColorGradingColorWheelPanel.cpp
+++ b/Engine/Plugins/Editor/ColorGrading/Source/ColorGradingEditor/Private/SColorGradingColorWheelPanel.cpp
@@ -50,144 +50,163 @@ void SColorGradingColorWheelPanel::Construct(const FArguments& InArgs)
 	for (int32 Index = 0; Index < NumColorWheels; ++Index)
     ChildSlot
 	[
-		SNew(SSplitter)
-		.Orientation(Orient_Horizontal)
-		.PhysicalSplitterHandleSize(1.0f)
-		.HitDetectionSplitterHandleSize(5.0f)
-		.Style(FAppStyle::Get(), "DetailsView.Splitter")
-
-		+ SSplitter::Slot()
-		.Value(0.8f)
+		SNew(SVerticalBox)
+
+		// Message indicating that multi select is unavailable in this panel
+		+ SVerticalBox::Slot()
 		[
-			SNew(SVerticalBox)
+			SNew(SBox)
+			.Visibility(this, &SColorGradingColorWheelPanel::GetMultiSelectWarningVisibility)
+			.HAlign(EHorizontalAlignment::HAlign_Center)
+			.VAlign(EVerticalAlignment::VAlign_Center)
+			[
+				SNew(STextBlock)
+				.Text(LOCTEXT("MultiSelectWarning", "Multi-select editing is unavailable in the Color Grading panel."))
+			]
+		]
+
+		// Color wheel panel
+		+ SVerticalBox::Slot()
+		[
+			SNew(SSplitter)
+			.Orientation(Orient_Horizontal)
+			.PhysicalSplitterHandleSize(1.0f)
+			.HitDetectionSplitterHandleSize(5.0f)
+			.Style(FAppStyle::Get(), "DetailsView.Splitter")
 			.Visibility(this, &SColorGradingColorWheelPanel::GetColorWheelPanelVisibility)
 
-			// Toolbar slot
-			+ SVerticalBox::Slot()
-			.AutoHeight()
-			.Padding(6, 4)
+			+ SSplitter::Slot()
+			.Value(0.8f)
 			[
-				SNew(SHorizontalBox)
-
-				+ SHorizontalBox::Slot()
-				.AutoWidth()
-				[
-					SAssignNew(ColorGradingGroupPropertyBox, SBox)
-				]
+				SNew(SVerticalBox)
 
-				+ SHorizontalBox::Slot()
-				.FillWidth(1.0f)
+				// Toolbar slot
+				+ SVerticalBox::Slot()
+				.AutoHeight()
+				.Padding(6, 4)
 				[
-					SNew(SSpacer)
-				]
+					SNew(SHorizontalBox)
 
-				+ SHorizontalBox::Slot()
-				.AutoWidth()
-				[
-					SAssignNew(ColorGradingElementsToolBarBox, SHorizontalBox)
-				]
+					+ SHorizontalBox::Slot()
+					.AutoWidth()
+					[
+						SAssignNew(ColorGradingGroupPropertyBox, SBox)
+					]
 
-				+ SHorizontalBox::Slot()
-				.FillWidth(1.0f)
-				[
-					SNew(SSpacer)
-				]
+					+ SHorizontalBox::Slot()
+					.FillWidth(1.0f)
+					[
+						SNew(SSpacer)
+					]
 
-				+ SHorizontalBox::Slot()
-				.AutoWidth()
-				[
-					MakeColorDisplayModeCheckbox()
-				]
+					+ SHorizontalBox::Slot()
+					.AutoWidth()
+					[
+						SAssignNew(ColorGradingElementsToolBarBox, SHorizontalBox)
+					]
 
-				+ SHorizontalBox::Slot()
-				.AutoWidth()
-				.Padding(5.0f, 0.0f, 0.0f, 0.0f)
-				.HAlign(HAlign_Right)
-				.VAlign(VAlign_Center)
-				[
-					SNew(SComboButton)
-					.ComboButtonStyle(&FAppStyle::Get().GetWidgetStyle<FComboButtonStyle>("SimpleComboButton"))
-					.OnGetMenuContent(this, &SColorGradingColorWheelPanel::MakeSettingsMenu)
-					.HasDownArrow(false)
-					.ContentPadding(FMargin(1.0f, 0.0f))
-					.ButtonContent()
+					+ SHorizontalBox::Slot()
+					.FillWidth(1.0f)
 					[
-						SNew(SImage)
-						.ColorAndOpacity(FSlateColor::UseForeground())
-						.Image(FAppStyle::Get().GetBrush("Icons.Settings"))
+						SNew(SSpacer)
 					]
-				]
-			]
 
-			+ SVerticalBox::Slot()
-			.FillHeight(1.0f)
-			.Padding(6, 4)
-			[
-				SNew(SHorizontalBox)
+					+ SHorizontalBox::Slot()
+					.AutoWidth()
+					[
+						MakeColorDisplayModeCheckbox()
+					]
 
-				+ SHorizontalBox::Slot()
-				.FillWidth(1.0f)
-				.HAlign(HAlign_Fill)
-				.Padding(2, 0)
-				[
-					SAssignNew(ColorWheels[0], SColorGradingColorWheel)
-					.ColorDisplayMode(this, &SColorGradingColorWheelPanel::GetColorDisplayMode)
-					.Orientation(ColorWheelOrientation)
-					.Visibility(this, &SColorGradingColorWheelPanel::GetColorWheelVisibility, 0)
+					+ SHorizontalBox::Slot()
+					.AutoWidth()
+					.Padding(5.0f, 0.0f, 0.0f, 0.0f)
+					.HAlign(HAlign_Right)
+					.VAlign(VAlign_Center)
+					[
+						SNew(SComboButton)
+						.ComboButtonStyle(&FAppStyle::Get().GetWidgetStyle<FComboButtonStyle>("SimpleComboButton"))
+						.OnGetMenuContent(this, &SColorGradingColorWheelPanel::MakeSettingsMenu)
+						.HasDownArrow(false)
+						.ContentPadding(FMargin(1.0f, 0.0f))
+						.ButtonContent()
+						[
+							SNew(SImage)
+							.ColorAndOpacity(FSlateColor::UseForeground())
+							.Image(FAppStyle::Get().GetBrush("Icons.Settings"))
+						]
+					]
 				]
 
-				+ SHorizontalBox::Slot()
-				.FillWidth(1.0f)
-				.HAlign(HAlign_Fill)
-				.Padding(2, 0)
+				+ SVerticalBox::Slot()
+				.FillHeight(1.0f)
+				.Padding(6, 4)
 				[
-					SAssignNew(ColorWheels[1], SColorGradingColorWheel)
-					.ColorDisplayMode(this, &SColorGradingColorWheelPanel::GetColorDisplayMode)
-					.Orientation(ColorWheelOrientation)
-					.Visibility(this, &SColorGradingColorWheelPanel::GetColorWheelVisibility, 1)
-				]
+					SNew(SHorizontalBox)
 
-				+ SHorizontalBox::Slot()
-				.FillWidth(1.0f)
-				.HAlign(HAlign_Fill)
-				.Padding(2, 0)
-				[
-					SAssignNew(ColorWheels[2], SColorGradingColorWheel)
-					.ColorDisplayMode(this, &SColorGradingColorWheelPanel::GetColorDisplayMode)
-					.Orientation(ColorWheelOrientation)
-					.Visibility(this, &SColorGradingColorWheelPanel::GetColorWheelVisibility, 2)
-				]
+					+ SHorizontalBox::Slot()
+					.FillWidth(1.0f)
+					.HAlign(HAlign_Fill)
+					.Padding(2, 0)
+					[
+						SAssignNew(ColorWheels[0], SColorGradingColorWheel)
+						.ColorDisplayMode(this, &SColorGradingColorWheelPanel::GetColorDisplayMode)
+						.Orientation(ColorWheelOrientation)
+						.Visibility(this, &SColorGradingColorWheelPanel::GetColorWheelVisibility, 0)
+					]
 
-				+ SHorizontalBox::Slot()
-				.FillWidth(1.0f)
-				.HAlign(HAlign_Fill)
-				.Padding(2, 0)
-				[
-					SAssignNew(ColorWheels[3], SColorGradingColorWheel)
-					.ColorDisplayMode(this, &SColorGradingColorWheelPanel::GetColorDisplayMode)
-					.Orientation(ColorWheelOrientation)
-					.Visibility(this, &SColorGradingColorWheelPanel::GetColorWheelVisibility, 3)
-				]
+					+ SHorizontalBox::Slot()
+					.FillWidth(1.0f)
+					.HAlign(HAlign_Fill)
+					.Padding(2, 0)
+					[
+						SAssignNew(ColorWheels[1], SColorGradingColorWheel)
+						.ColorDisplayMode(this, &SColorGradingColorWheelPanel::GetColorDisplayMode)
+						.Orientation(ColorWheelOrientation)
+						.Visibility(this, &SColorGradingColorWheelPanel::GetColorWheelVisibility, 1)
+					]
 
-				+ SHorizontalBox::Slot()
-				.FillWidth(1.0f)
-				.HAlign(HAlign_Fill)
-				.Padding(2, 0)
-				[
-					SAssignNew(ColorWheels[4], SColorGradingColorWheel)
-					.ColorDisplayMode(this, &SColorGradingColorWheelPanel::GetColorDisplayMode)
-					.Orientation(ColorWheelOrientation)
-					.Visibility(this, &SColorGradingColorWheelPanel::GetColorWheelVisibility, 4)
+					+ SHorizontalBox::Slot()
+					.FillWidth(1.0f)
+					.HAlign(HAlign_Fill)
+					.Padding(2, 0)
+					[
+						SAssignNew(ColorWheels[2], SColorGradingColorWheel)
+						.ColorDisplayMode(this, &SColorGradingColorWheelPanel::GetColorDisplayMode)
+						.Orientation(ColorWheelOrientation)
+						.Visibility(this, &SColorGradingColorWheelPanel::GetColorWheelVisibility, 2)
+					]
+
+					+ SHorizontalBox::Slot()
+					.FillWidth(1.0f)
+					.HAlign(HAlign_Fill)
+					.Padding(2, 0)
+					[
+						SAssignNew(ColorWheels[3], SColorGradingColorWheel)
+						.ColorDisplayMode(this, &SColorGradingColorWheelPanel::GetColorDisplayMode)
+						.Orientation(ColorWheelOrientation)
+						.Visibility(this, &SColorGradingColorWheelPanel::GetColorWheelVisibility, 3)
+					]
+
+					+ SHorizontalBox::Slot()
+					.FillWidth(1.0f)
+					.HAlign(HAlign_Fill)
+					.Padding(2, 0)
+					[
+						SAssignNew(ColorWheels[4], SColorGradingColorWheel)
+						.ColorDisplayMode(this, &SColorGradingColorWheelPanel::GetColorDisplayMode)
+						.Orientation(ColorWheelOrientation)
+						.Visibility(this, &SColorGradingColorWheelPanel::GetColorWheelVisibility, 4)
+					]
 				]
 			]
-		]
 
-		+ SSplitter::Slot()
-		.Value(0.2f)
-		[
-			SAssignNew(DetailView, SColorGradingDetailView)
-			.PropertyRowGeneratorSource(ColorGradingDataModel->GetPropertyRowGenerator())
-			.OnFilterDetailTreeNode(this, &SColorGradingColorWheelPanel::FilterDetailTreeNode)
+			+ SSplitter::Slot()
+			.Value(0.2f)
+			[
+				SAssignNew(DetailView, SColorGradingDetailView)
+				.PropertyRowGeneratorSource(ColorGradingDataModel->GetPropertyRowGenerator())
+				.OnFilterDetailTreeNode(this, &SColorGradingColorWheelPanel::FilterDetailTreeNode)
+			]
 		]
 	];
 }
@@ -674,8 +693,14 @@ ECheckBoxState SColorGradingColorWheelPanel::IsColorGradingElementSelected(FText
 
 EVisibility SColorGradingColorWheelPanel::GetColorWheelPanelVisibility() const
 {
-	bool bHasObjects = ColorGradingDataModel && ColorGradingDataModel->GetPropertyRowGenerator()->GetSelectedObjects().Num() > 0;
-	return bHasObjects ? EVisibility::Visible : EVisibility::Collapsed;
+	bool bHasObject = ColorGradingDataModel && ColorGradingDataModel->GetPropertyRowGenerator()->GetSelectedObjects().Num() == 1;
+	return bHasObject ? EVisibility::Visible : EVisibility::Collapsed;
+}
+
+EVisibility SColorGradingColorWheelPanel::GetMultiSelectWarningVisibility() const
+{
+	bool bHasMultipleObjects = ColorGradingDataModel && ColorGradingDataModel->GetPropertyRowGenerator()->GetSelectedObjects().Num() > 1;
+	return bHasMultipleObjects ? EVisibility::Visible : EVisibility::Collapsed;
 }
 
 EVisibility SColorGradingColorWheelPanel::GetColorWheelVisibility(int32 ColorWheelIndex) const
diff --git a/Engine/Plugins/Editor/ColorGrading/Source/ColorGradingEditor/Private/SColorGradingColorWheelPanel.h b/Engine/Plugins/Editor/ColorGrading/Source/ColorGradingEditor/Private/SColorGradingColorWheelPanel.h
--- a/Engine/Plugins/Editor/ColorGrading/Source/ColorGradingEditor/Private/SColorGradingColorWheelPanel.h
+++ b/Engine/Plugins/Editor/ColorGrading/Source/ColorGradingEditor/Private/SColorGradingColorWheelPanel.h
@@ -75,6 +75,7 @@ private:
 
 	ECheckBoxState IsColorGradingElementSelected(FText ElementName) const;
 	EVisibility GetColorWheelPanelVisibility() const;
+	EVisibility GetMultiSelectWarningVisibility() const;
 	EVisibility GetColorWheelVisibility(int32 ColorWheelIndex) const;
 
 	EColorGradingColorDisplayMode GetColorDisplayMode() const { return ColorDisplayMode; }
diff --git a/Engine/Plugins/Editor/ColorGrading/Source/ColorGradingEditor/Private/SColorGradingObjectList.cpp b/Engine/Plugins/Editor/ColorGrading/Source/ColorGradingEditor/Private/SColorGradingObjectList.cpp
deleted file mode 100644
diff --git a/Engine/Plugins/Editor/ColorGrading/Source/ColorGradingEditor/Private/SColorGradingObjectList.h b/Engine/Plugins/Editor/ColorGrading/Source/ColorGradingEditor/Private/SColorGradingObjectList.h
deleted file mode 100644
diff --git a/Engine/Plugins/Editor/ColorGrading/Source/ColorGradingEditor/Private/SColorGradingPanel.cpp b/Engine/Plugins/Editor/ColorGrading/Source/ColorGradingEditor/Private/SColorGradingPanel.cpp
--- a/Engine/Plugins/Editor/ColorGrading/Source/ColorGradingEditor/Private/SColorGradingPanel.cpp
+++ b/Engine/Plugins/Editor/ColorGrading/Source/ColorGradingEditor/Private/SColorGradingPanel.cpp
@@ -4,10 +4,13 @@
 
 #include "ColorGradingCommands.h"
 #include "ColorGradingEditorDataModel.h"
+#include "ColorGradingMixerObjectFilter.h"
 #include "IColorGradingEditor.h"
 #include "SColorGradingColorWheelPanel.h"
-#include "SColorGradingObjectList.h"
 
+#include "ActorTreeItem.h"
+#include "Algo/Compare.h"
+#include "ComponentTreeItem.h"
 #include "ColorCorrectRegion.h"
 #include "Editor.h"
 #include "Engine/Blueprint.h"
@@ -16,6 +19,8 @@
 #include "Framework/Commands/GenericCommands.h"
 #include "Framework/MultiBox/MultiBoxBuilder.h"
 #include "LevelEditorSubsystem.h"
+#include "Views/List/ObjectMixerEditorList.h"
+#include "Views/List/SObjectMixerEditorList.h"
 #include "Widgets/Input/SButton.h"
 #include "Widgets/Layout/SBox.h"
 #include "Widgets/Layout/SExpandableArea.h"
@@ -29,19 +34,6 @@
 
 SColorGradingPanel::~SColorGradingPanel()
 {
-	for (const FColorGradingListItemRef& ColorGradingItem : ColorGradingItemList)
-	{
-		if (ColorGradingItem->Component.IsValid())
-		{
-			UnbindBlueprintCompiledDelegate(ColorGradingItem->Component->GetClass());
-		}
-
-		if (ColorGradingItem->Actor.IsValid())
-		{
-			UnbindBlueprintCompiledDelegate(ColorGradingItem->Actor->GetClass());
-		}
-	}
-
 	FCoreUObjectDelegates::OnObjectsReplaced.RemoveAll(this);
 
 	if (GEngine)
@@ -66,11 +58,32 @@ void SColorGradingPanel::Construct(const FArguments& InArgs)
 	ColorGradingDataModel = MakeShared<FColorGradingEditorDataModel>();
 	ColorGradingDataModel->OnDataModelGenerated().AddSP(this, &SColorGradingPanel::OnColorGradingDataModelGenerated);
 
-	FCoreUObjectDelegates::OnObjectsReplaced.AddSP(this, &SColorGradingPanel::OnObjectsReplaced);
-	GEngine->OnLevelActorAdded().AddSP(this, &SColorGradingPanel::OnLevelActorAdded);
-	GEngine->OnLevelActorDeleted().AddSP(this, &SColorGradingPanel::OnLevelActorDeleted);
-	GEngine->OnWorldAdded().AddSP(this, &SColorGradingPanel::OnWorldAdded);
-	GEngine->OnWorldDestroyed().AddSP(this, &SColorGradingPanel::OnWorldDestroyed);
+	const FName ModuleName = FName(TEXT("ColorGrading"));
+	ObjectListModel = MakeShared<FObjectMixerEditorList>(ModuleName, InArgs._SelectionInterface);
+	ObjectListModel->Initialize();
+	ObjectListModel->SetDefaultFilterClass(UColorGradingMixerObjectFilter::StaticClass());
+
+	TSharedRef<SWidget> ObjectListWidget = ObjectListModel->GetOrCreateWidget();
+	ColorGradingObjectListView = StaticCastSharedRef<SSceneOutliner>(ObjectListWidget).ToSharedPtr();
+
+	TSharedRef<SObjectMixerEditorList> ObjectMixerList = StaticCastSharedRef<SObjectMixerEditorList>(ObjectListWidget);
+	ObjectMixerList->GetOnItemSelectionChanged().AddSP(this, &SColorGradingPanel::OnListSelectionChanged);
+	ObjectMixerList->GetOnSelectionSynchronized().AddSP(this, &SColorGradingPanel::OnListSelectionSynchronized);
+
+	ObjectMixerList->AddFilter(
+		MakeShared<TSceneOutlinerPredicateFilter<FActorTreeItem>>(
+			FActorTreeItem::FFilterPredicate::CreateSPLambda(this, [this](const AActor* Actor)
+			{
+				if (ActorFilter)
+				{
+					return ActorFilter(Actor);
+				}
+
+				return true;
+			}),
+			FSceneOutlinerFilter::EDefaultBehaviour::Pass
+		)
+	);
 
 	GEditor->RegisterForUndo(this);
 
@@ -91,51 +104,18 @@ void SColorGradingPanel::Construct(const FArguments& InArgs)
 
 			// Splitter slot for object list
 			+SSplitter::Slot()
-			.Value(0.12f)
+			.Value(0.2f)
 			[
 				SNew(SBox)
 				.Padding(FMargin(4.f))
 				[
-					SNew(SBorder)
-					.Padding(FMargin(0.0f))
-					.BorderImage(FAppStyle::GetBrush("Brushes.Recessed"))
-					[
-						SNew(SScrollBox)
-						+SScrollBox::Slot()
-						[
-							SNew(SExpandableArea)
-							.BorderImage(FAppStyle::Get().GetBrush("Brushes.Header"))
-							.BodyBorderImage(FAppStyle::Get().GetBrush("Brushes.Recessed"))
-							.HeaderPadding(FMargin(4.0f, 2.0f))
-							.InitiallyCollapsed(false)
-							.AllowAnimatedTransition(false)
-							.Visibility_Lambda([this]() { return ColorGradingItemList.Num() ? EVisibility::Visible : EVisibility::Collapsed; })
-							.HeaderContent()
-							[
-								SNew(SBox)
-								.HeightOverride(24.0f)
-								.VAlign(VAlign_Center)
-								[
-									SNew(STextBlock)
-										.Text(LOCTEXT("ColorGradingObjectListLabel", "Objects"))
-										.TextStyle(FAppStyle::Get(), "ButtonText")
-										.Font(FAppStyle::Get().GetFontStyle("NormalFontBold"))
-								]
-							]
-							.BodyContent()
-							[
-								SAssignNew(ColorGradingObjectListView, SColorGradingObjectList)
-								.ColorGradingItemsSource(&ColorGradingItemList)
-								.OnSelectionChanged(this, &SColorGradingPanel::OnListSelectionChanged)
-							]
-						]
-					]
+					ObjectListWidget
 				]
 			]
 
 			// Splitter slot for color grading controls/details
 			+SSplitter::Slot()
-			.Value(0.88f)
+			.Value(0.8f)
 			[
 				SNew(SVerticalBox)
 
@@ -205,18 +185,6 @@ void SColorGradingPanel::Construct(const FArguments& InArgs)
 	];
 }
 
-void SColorGradingPanel::Tick(const FGeometry& AllottedGeometry, const double InCurrentTime, const float InDeltaTime)
-{
-	SCompoundWidget::Tick(AllottedGeometry, InCurrentTime, InDeltaTime);
-
-	if (bRefreshOnNextTick)
-	{
-		Refresh();
-
-		bRefreshOnNextTick = false;
-	}
-}
-
 void SColorGradingPanel::Refresh()
 {
 	FColorGradingPanelState PanelState;
@@ -245,19 +213,30 @@ void SColorGradingPanel::GetPanelState(FColorGradingPanelState& OutPanelState) c
 
 	if (ColorGradingObjectListView.IsValid())
 	{
-		TArray<FColorGradingListItemRef> SelectedItems = ColorGradingObjectListView->GetSelectedItems();
+		TArray<FSceneOutlinerTreeItemPtr> SelectedItems = ColorGradingObjectListView->GetSelectedItems();
 
-		for (const FColorGradingListItemRef& SelectedItem : SelectedItems)
+		for (const FSceneOutlinerTreeItemPtr& SelectedItem : SelectedItems)
 		{
-			if (SelectedItem.IsValid())
+			if (!SelectedItem.IsValid())
 			{
-				if (SelectedItem->Component.IsValid())
+				continue;
+			}
+			
+			if (const FComponentTreeItem* ComponentTreeItem = SelectedItem->CastTo<FComponentTreeItem>())
+			{
+				if (ComponentTreeItem->Component.IsValid())
 				{
-					OutPanelState.SelectedObjects.Add(SelectedItem->Component);
+					OutPanelState.SelectedObjects.Add(ComponentTreeItem->Component.Get());
+					continue;
 				}
-				else if (SelectedItem->Actor.IsValid())
+			}
+			
+			if (const FActorTreeItem* ActorTreeItem = SelectedItem->CastTo<FActorTreeItem>())
+			{
+				if (ActorTreeItem->Actor.IsValid())
 				{
-					OutPanelState.SelectedObjects.Add(SelectedItem->Actor);
+					OutPanelState.SelectedObjects.Add(ActorTreeItem->Actor.Get());
+					continue;
 				}
 			}
 		}
@@ -266,33 +245,28 @@ void SColorGradingPanel::GetPanelState(FColorGradingPanelState& OutPanelState) c
 
 void SColorGradingPanel::SetPanelState(const FColorGradingPanelState& InPanelState)
 {
-	TArray<FColorGradingListItemRef> ItemsToSelect;
+	TArray<FSceneOutlinerTreeItemPtr> ItemsToSelect;
 
 	for (const TWeakObjectPtr<UObject>& SelectedObject : InPanelState.SelectedObjects)
 	{
-		if (SelectedObject.IsValid())
+		if (!SelectedObject.IsValid())
 		{
-			auto FindColorGradingItem = [&SelectedObject](const FColorGradingListItemRef& ColorGradingItem)
-			{
-				if (SelectedObject->IsA<AActor>())
-				{
-					return ColorGradingItem->Actor == SelectedObject && ColorGradingItem->Component == nullptr;
-				}
-
-				return ColorGradingItem->Actor == SelectedObject || ColorGradingItem->Component == SelectedObject;
-			};
+			continue;
+		}
 
-			if (FColorGradingListItemRef* FoundItem = ColorGradingItemList.FindByPredicate(FindColorGradingItem))
-			{
-				ItemsToSelect.Add(*FoundItem);
-				break;
-			}
+		FSceneOutlinerTreeItemPtr Item = ColorGradingObjectListView->GetTreeItem(SelectedObject.Get());
+		if (Item.IsValid())
+		{
+			ItemsToSelect.Add(Item);
+			break;
 		}
+
 	}
 
 	if (!ItemsToSelect.IsEmpty() && ColorGradingObjectListView.IsValid())
 	{
-		ColorGradingObjectListView->SetSelectedItems(ItemsToSelect);
+		ColorGradingObjectListView->ClearSelection();
+		ColorGradingObjectListView->SetItemSelection(ItemsToSelect, true);
 	}
 
 	ColorGradingDataModel->SetPanelState(InPanelState);
@@ -373,90 +347,11 @@ UWorld* SColorGradingPanel::GetWorld()
 	return Level->GetWorld();
 }
 
-void SColorGradingPanel::BindBlueprintCompiledDelegate(const UClass* Class)
-{
-	if (UBlueprint* Blueprint = UBlueprint::GetBlueprintFromClass(Class))
-	{
-		if (!Blueprint->OnCompiled().IsBoundToObject(this))
-		{
-			Blueprint->OnCompiled().AddSP(this, &SColorGradingPanel::OnBlueprintCompiled);
-		}
-	}
-}
-
-void SColorGradingPanel::UnbindBlueprintCompiledDelegate(const UClass* Class)
-{
-	if (UBlueprint* Blueprint = UBlueprint::GetBlueprintFromClass(Class))
-	{
-		Blueprint->OnCompiled().RemoveAll(this);
-	}
-}
-
 void SColorGradingPanel::RefreshColorGradingList()
 {
-	for (const FColorGradingListItemRef& Item : ColorGradingItemList)
-	{
-		if (Item->Component.IsValid())
-		{
-			UnbindBlueprintCompiledDelegate(Item->Component->GetClass());
-		}
-
-		if (Item->Actor.IsValid())
-		{
-			UnbindBlueprintCompiledDelegate(Item->Actor->GetClass());
-		}
-	}
-
-	ColorGradingItemList.Empty();
-
-	UWorld* World = GetWorld();
-	if (World && ColorGradingDataModel.IsValid())
-	{
-		// Sorter that sorts the list items alphabetically by their display name
-		auto AlphabeticalSort = [](const FColorGradingListItemRef& A, const FColorGradingListItemRef& B)
-		{
-			if (A.IsValid() && B.IsValid())
-			{
-				return *A < *B;
-			}
-			else
-			{
-				return false;
-			}
-		};
-
-		// Add all relevant actors
-		const TSet<TSubclassOf<AActor>>& ListItemClasses = FColorGradingListItem::GetActorClassesWithListItemGenerators();
-
-		for (TSubclassOf<AActor> ActorClass : ListItemClasses)
-		{
-			for (TActorIterator<AActor> ActorIter(World, ActorClass.Get()); ActorIter; ++ActorIter)
-			{
-				AActor* Actor = *ActorIter;
-
-				if (ActorFilter && !ActorFilter(Actor))
-				{
-					continue;
-				}
-
-				TArray<FColorGradingListItemRef> ListItems = FColorGradingListItem::GenerateColorGradingListItems(Actor);
-
-				if (ListItems.IsEmpty())
-				{
-					continue;
-				}
-
-				BindBlueprintCompiledDelegate(Actor->GetClass());
-				ColorGradingItemList.Append(ListItems);
-			}
-		}
-
-		ColorGradingItemList.Sort(AlphabeticalSort);
-	}
-
-	if (ColorGradingObjectListView.IsValid())
+	if (ObjectListModel)
 	{
-		ColorGradingObjectListView->RefreshList();
+		ObjectListModel->RequestRebuildList();
 	}
 }
 
@@ -623,138 +518,59 @@ void SColorGradingPanel::OnColorGradingGroupRenamed(const FText& InText, ETextCo
 	ColorGradingDataModel->OnColorGradingGroupRenamed().Broadcast(GroupIndex, InText);
 }
 
-void SColorGradingPanel::OnObjectsReplaced(const TMap<UObject*, UObject*>& OldToNewInstanceMap)
+void SColorGradingPanel::OnColorGradingDataModelGenerated()
 {
-	bool bNeedsFullRefresh = false;
-	bool bNeedsListRefresh = false;
-
-	const TArray<TWeakObjectPtr<UObject>> SelectedObjects = ColorGradingDataModel->GetObjects();
-
-	for (const TPair<UObject*, UObject*>& Pair : OldToNewInstanceMap)
-	{
-		if (Pair.Key && Pair.Value)
-		{
-			FColorGradingListItemRef* FoundColorGradingItemPtr = nullptr;
-
-			// Must use GetEvenIfUnreachable on the weak pointers here because most of the time, the objects being replaced have already been marked for GC, and TWeakObjectPtr
-			// will return nullptr from Get on GC-marked objects
-			FoundColorGradingItemPtr = ColorGradingItemList.FindByPredicate([&Pair](const FColorGradingListItemRef& ColorGradingItem)
-			{
-				return ColorGradingItem->Actor.GetEvenIfUnreachable() == Pair.Key || ColorGradingItem->Component.GetEvenIfUnreachable() == Pair.Key;
-			});
-
-			if (FoundColorGradingItemPtr)
-			{
-				FColorGradingListItemRef FoundColorGradingItem = *FoundColorGradingItemPtr;
-				if (FoundColorGradingItem->Actor.GetEvenIfUnreachable() == Pair.Key)
-				{
-					FoundColorGradingItem->Actor = Cast<AActor>(Pair.Value);
-				}
-				else if (FoundColorGradingItem->Component.GetEvenIfUnreachable() == Pair.Key)
-				{
-					FoundColorGradingItem->Component = Cast<UActorComponent>(Pair.Value);
-				}
-
-				bNeedsListRefresh = true;
-			}
-
-			if (SelectedObjects.Contains(Pair.Key))
-			{
-				bNeedsFullRefresh = true;
-			}
-		}
-	}
+	FillColorGradingGroupToolBar();
 
-	if (bNeedsFullRefresh)
-	{
-		Refresh();
-	}
-	else if (bNeedsListRefresh && ColorGradingObjectListView)
+	if (ColorWheelPanel.IsValid())
 	{
-		ColorGradingObjectListView->RefreshList();
+		ColorWheelPanel->Refresh();
 	}
 }
 
-void SColorGradingPanel::OnLevelActorAdded(AActor* Actor)
+void SColorGradingPanel::OnListSelectionChanged(FSceneOutlinerTreeItemPtr TreeItem, ESelectInfo::Type Type)
 {
-	// Only refresh when the actor being added is being added to the edited world
-	if (UWorld* World = GetWorld())
+	if (Type == ESelectInfo::Direct)
 	{
-		if (World == Actor->GetWorld())
-		{
-			if (Actor->IsA<APostProcessVolume>() || Actor->IsA<AColorCorrectRegion>())
-			{
-				// Wait to refresh, as this event can be fired off for several actors in a row in certain cases, such as when the root actor is recompiled after a property change
-				bRefreshOnNextTick = true;
-			}
-		}
+		return;
 	}
-}
 
-void SColorGradingPanel::OnLevelActorDeleted(AActor* Actor)
-{
-	auto ContainsActorRef = [Actor](const FColorGradingListItemRef& ColorGradingItem)
-	{
-		return ColorGradingItem->Actor.GetEvenIfUnreachable() == Actor;
-	};
-
-	if (ColorGradingItemList.ContainsByPredicate(ContainsActorRef))
-	{
-		// Must wait for next tick to refresh because the actor has not actually been removed from the level at this point
-		bRefreshOnNextTick = true;
-	}
+	UpdateSelectionFromList();
 }
 
-void SColorGradingPanel::OnWorldAdded(UWorld* World)
+void SColorGradingPanel::OnListSelectionSynchronized()
 {
-	bRefreshOnNextTick = true;
+	UpdateSelectionFromList();
 }
 
-void SColorGradingPanel::OnWorldDestroyed(UWorld* World)
+void SColorGradingPanel::UpdateSelectionFromList()
 {
-	auto NoLongerValid = [World](const FColorGradingListItemRef& ColorGradingItem)
-	{
-		return !ColorGradingItem->Actor.IsValid() || ColorGradingItem->Actor->GetWorld() == World;
-	};
-
-	// Immediately remove any affected actors so we don't try to access them
-	ColorGradingItemList.RemoveAll(NoLongerValid);
-}
+	TArray<FSceneOutlinerTreeItemPtr> SelectedOutlinerItems = ColorGradingObjectListView->GetSelectedItems();
+	TArray<UObject*> SelectedObjects;
 
-void SColorGradingPanel::OnBlueprintCompiled(UBlueprint* Blueprint)
-{
-	Refresh();
-}
-
-void SColorGradingPanel::OnColorGradingDataModelGenerated()
-{
-	FillColorGradingGroupToolBar();
-
-	if (ColorWheelPanel.IsValid())
-	{
-		ColorWheelPanel->Refresh();
-	}
-}
-
-void SColorGradingPanel::OnListSelectionChanged(TSharedRef<SColorGradingObjectList> SourceList, FColorGradingListItemRef SelectedItem, ESelectInfo::Type SelectInfo)
-{
-	if (SelectInfo != ESelectInfo::Direct)
+	for (FSceneOutlinerTreeItemPtr OutlinerItems : SelectedOutlinerItems)
 	{
-		TArray<FColorGradingListItemRef> SelectedObjects = SourceList->GetSelectedItems();
-		TArray<UObject*> ObjectsToColorGrade;
-		for (const FColorGradingListItemRef& SelectedObject : SelectedObjects)
+		if (const FActorTreeItem* ActorTreeItem = OutlinerItems->CastTo<FActorTreeItem>())
 		{
-			if (SelectedObject->Component.IsValid())
+			if (ActorTreeItem->Actor.IsValid())
 			{
-				ObjectsToColorGrade.Add(SelectedObject->Component.Get());
+				SelectedObjects.Add(ActorTreeItem->Actor.Get());
 			}
-			else if (SelectedObject->Actor.IsValid())
+		}
+
+		if (const FComponentTreeItem* ComponentTreeItem = OutlinerItems->CastTo<FComponentTreeItem>())
+		{
+			if (ComponentTreeItem->Component.IsValid())
 			{
-				ObjectsToColorGrade.Add(SelectedObject->Actor.Get());
+				SelectedObjects.Add(ComponentTreeItem->Component.Get());
 			}
 		}
+	}
 
-		SetSelectedObjects(ObjectsToColorGrade);
+	TArray<TWeakObjectPtr<UObject>> OldSelectedObjects = ColorGradingDataModel->GetObjects();
+	if (!Algo::Compare(OldSelectedObjects, SelectedObjects))
+	{
+		SetSelectedObjects(SelectedObjects);
 	}
 }
 
diff --git a/Engine/Plugins/Editor/ColorGrading/Source/ColorGradingEditor/Public/ColorGradingMixerObjectFilterRegistry.h b/Engine/Plugins/Editor/ColorGrading/Source/ColorGradingEditor/Public/ColorGradingMixerObjectFilterRegistry.h
new file mode 100644
--- /dev/null
+++ b/Engine/Plugins/Editor/ColorGrading/Source/ColorGradingEditor/Public/ColorGradingMixerObjectFilterRegistry.h
@@ -0,0 +1,42 @@
+// Copyright Epic Games, Inc. All Rights Reserved.
+
+#pragma once
+
+#include "GameFramework/Actor.h"
+#include "HAL/Platform.h"
+#include "Templates/SubclassOf.h"
+
+class UClass;
+
+/**
+ * Contains functions for controlling which classes can be managed from the Color Grading panel's ObjectMixer-based hierarchy panel
+ */
+struct COLORGRADINGEDITOR_API FColorGradingMixerObjectFilterRegistry
+{
+public:
+
+	/** Register an object class that can be seen in a Color Grading panel's object list */
+	static void RegisterObjectClassToFilter(UClass* Class)
+	{
+		ObjectClassesToFilter.Add(Class);
+	}
+
+	/** Register an actor class that can be placed from the Color Grading panel's object list */
+	static void RegisterActorClassToPlace(TSubclassOf<AActor> Class)
+	{
+		ActorClassesToPlace.Add(Class);
+	}
+
+	/** Get the set of object classes that can be seen in a Color Grading panel's object list */
+	static const TSet<UClass*>& GetObjectClassesToFilter() { return ObjectClassesToFilter; }
+
+	/** Get the set of actor classes that can be placed from a Color Grading panel's object list */
+	static const TSet<TSubclassOf<AActor>>& GetActorClassesToPlace() { return ActorClassesToPlace; }
+
+private:
+	/** Set of classes that can be seen in the object panel */
+	inline static TSet<UClass*> ObjectClassesToFilter;
+
+	/** Set of classes that can be placed from the object panel */
+	inline static TSet<TSubclassOf<AActor>> ActorClassesToPlace;
+};
diff --git a/Engine/Plugins/Editor/ColorGrading/Source/ColorGradingEditor/Public/SColorGradingPanel.h b/Engine/Plugins/Editor/ColorGrading/Source/ColorGradingEditor/Public/SColorGradingPanel.h
--- a/Engine/Plugins/Editor/ColorGrading/Source/ColorGradingEditor/Public/SColorGradingPanel.h
+++ b/Engine/Plugins/Editor/ColorGrading/Source/ColorGradingEditor/Public/SColorGradingPanel.h
@@ -5,18 +5,20 @@
 #include "CoreMinimal.h"
 #include "Widgets/SCompoundWidget.h"
 #include "EditorUndoClient.h"
+#include "SceneOutlinerFwd.h"
+#include "SelectionInterface/IObjectMixerSelectionInterface.h"
 
 #include "ColorGradingPanelState.h"
 #include "ColorGradingListItem.h"
 
 class FColorGradingEditorDataModel;
+class FObjectMixerEditorList;
 class SColorGradingColorWheelPanel;
-class SColorGradingObjectList;
 class SHorizontalBox;
 class SInlineEditableTextBlock;
 class UWorld;
 
-using FColorGradingActorFilter = TFunction<bool(AActor*)>;
+using FColorGradingActorFilter = TFunction<bool(const AActor*)>;
 
 /** Main panel of a color grading drawer widget, which displays color wheels or selected object details */
 class COLORGRADINGEDITOR_API SColorGradingPanel : public SCompoundWidget, public FEditorUndoClient
@@ -40,14 +42,13 @@ public:
 		/** Function which, if it returns false when passed an actor, filters it and its sub-entries out of the color grading item list */
 		SLATE_ARGUMENT(FColorGradingActorFilter, ActorFilter)
 
+		/** Optional interface which, if provided, will determine how objects selected in this panel will be synchronized with the rest of the editor */
+		SLATE_ARGUMENT(TSharedPtr<IObjectMixerSelectionInterface>, SelectionInterface)
+
 	SLATE_END_ARGS()
 
 	void Construct(const FArguments& InArgs);
 
-	//~ SWidget interface
-	virtual void Tick(const FGeometry& AllottedGeometry, const double InCurrentTime, const float InDeltaTime) override;
-	//~ End SWidget interface
-
 	/** Refreshes the panel's UI to match the current state of the level */
 	void Refresh();
 
@@ -67,12 +68,6 @@ private:
 	/** Get the world currently being edited */
 	UWorld* GetWorld();
 
-	/** Binds a callback to the BlueprintCompiled delegate of the specified class */
-	void BindBlueprintCompiledDelegate(const UClass* Class);
-
-	/** Unbinds a callback to the BlueprintCompiled delegate of the specified class */
-	void UnbindBlueprintCompiledDelegate(const UClass* Class);
-
 	/** Refreshes the object list, filling it with the current color gradable objects from the root actor and world */
 	void RefreshColorGradingList();
 
@@ -106,24 +101,6 @@ private:
 	/** Raised when a rename has been committed on a color grading group */
 	void OnColorGradingGroupRenamed(const FText& InText, ETextCommit::Type TextCommitType, int32 GroupIndex);
 
-	/** Raised when the editor replaces any UObjects with new instantiations, usually when actors have been recompiled from blueprints */
-	void OnObjectsReplaced(const TMap<UObject*, UObject*>& OldToNewInstanceMap);
-
-	/** Raised when an actor is added to the current level */
-	void OnLevelActorAdded(AActor* Actor);
-
-	/** Raised when an actor has been deleted from the currnent level */
-	void OnLevelActorDeleted(AActor* Actor);
-
-	/** Raised when a world has been added */
-	void OnWorldAdded(UWorld* World);
-
-	/** Raised when a world has been destroyed */
-	void OnWorldDestroyed(UWorld* World);
-
-	/** Raised when the specified blueprint has been recompiled */
-	void OnBlueprintCompiled(UBlueprint* Blueprint);
-
 	/** Raised when the color grading data model has been generated */
 	void OnColorGradingDataModelGenerated();
 
@@ -131,10 +108,19 @@ private:
 	FReply DockInLayout();
 
 	/** Raised when the user has selected a new item in any of the drawer's list views */
-	void OnListSelectionChanged(TSharedRef<SColorGradingObjectList> SourceList, FColorGradingListItemRef SelectedItem, ESelectInfo::Type SelectInfo);
+	void OnListSelectionChanged(FSceneOutlinerTreeItemPtr TreeItem, ESelectInfo::Type Type);
+
+	/** Raised when the drawer's outliner has been synchronized with the editor selection */
+	void OnListSelectionSynchronized();
+
+	/** Update the selected items to match the list view */
+	void UpdateSelectionFromList();
 
 
 private:
+	/** Model for the object mixer list used to display the color gradable object hierarchy. */
+	TSharedPtr<FObjectMixerEditorList> ObjectListModel;
+
 	/** Box containing the color grading groups */
 	TSharedPtr<SHorizontalBox> ColorGradingGroupToolBarBox;
 
@@ -148,10 +134,7 @@ private:
 	TAttribute<UWorld*> OverrideWorld;
 
 	/** Color grading object list widget being displayed in the drawer's list panel */
-	TSharedPtr<SColorGradingObjectList> ColorGradingObjectListView;
-
-	/** Source list for the color grading object list widget */
-	TArray<FColorGradingListItemRef> ColorGradingItemList;
+	TSharedPtr<SSceneOutliner> ColorGradingObjectListView;
 
 	/** The color grading data model for the currently selected objects */
 	TSharedPtr<FColorGradingEditorDataModel> ColorGradingDataModel;
diff --git a/Engine/Plugins/Editor/ObjectMixer/ObjectMixer/Source/ObjectMixer/Private/ObjectFilter/ObjectMixerEditorObjectFilter.cpp b/Engine/Plugins/Editor/ObjectMixer/ObjectMixer/Source/ObjectMixer/Private/ObjectFilter/ObjectMixerEditorObjectFilter.cpp
--- a/Engine/Plugins/Editor/ObjectMixer/ObjectMixer/Source/ObjectMixer/Private/ObjectFilter/ObjectMixerEditorObjectFilter.cpp
+++ b/Engine/Plugins/Editor/ObjectMixer/ObjectMixer/Source/ObjectMixer/Private/ObjectFilter/ObjectMixerEditorObjectFilter.cpp
@@ -20,6 +20,16 @@ bool UObjectMixerObjectFilter::GetShowTransientObjects() const
 	return false;
 }
 
+bool UObjectMixerObjectFilter::ShouldAllowHybridRows() const
+{
+	return true;
+}
+
+bool UObjectMixerObjectFilter::ShouldAllowColumnCustomizationByUser() const
+{
+	return true;
+}
+
 TSet<FName> UObjectMixerObjectFilter::GetColumnsToShowByDefault() const
 {
 	return {};
diff --git a/Engine/Plugins/Editor/ObjectMixer/ObjectMixer/Source/ObjectMixer/Private/ObjectMixerEditorModule.cpp b/Engine/Plugins/Editor/ObjectMixer/ObjectMixer/Source/ObjectMixer/Private/ObjectMixerEditorModule.cpp
--- a/Engine/Plugins/Editor/ObjectMixer/ObjectMixer/Source/ObjectMixer/Private/ObjectMixerEditorModule.cpp
+++ b/Engine/Plugins/Editor/ObjectMixer/ObjectMixer/Source/ObjectMixer/Private/ObjectMixerEditorModule.cpp
@@ -11,11 +11,8 @@
 
 #include "Engine/Level.h"
 #include "ISettingsModule.h"
-#include "LevelEditor.h"
 #include "LevelEditorSequencerIntegration.h"
 #include "Misc/CoreDelegates.h"
-#include "Misc/TransactionObjectEvent.h"
-#include "Selection.h"
 #include "ToolMenus.h"
 #include "Widgets/Docking/SDockTab.h"
 #include "WorkspaceMenuStructure.h"
@@ -54,8 +51,6 @@ UWorld* FObjectMixerEditorModule::GetWorld()
 
 void FObjectMixerEditorModule::Initialize()
 {	
-	BindDelegates();
-	
 	SetupMenuItemVariables();
 	
 	RegisterTabSpawner();
@@ -64,51 +59,6 @@ void FObjectMixerEditorModule::Initialize()
 
 void FObjectMixerEditorModule::Teardown()
 {
-	// Unbind Delegates
-	if (GEngine)
-	{
-		GEngine->OnLevelActorFolderChanged().RemoveAll(this);
-		GEngine->OnLevelActorListChanged().RemoveAll(this);
-		GEngine->OnActorFolderAdded().RemoveAll(this);
-		GEngine->OnLevelActorDetached().RemoveAll(this);
-		GEngine->OnLevelActorAttached().RemoveAll(this);
-		GEngine->OnActorFolderAdded().RemoveAll(this);
-		GEngine->OnActorFoldersUpdatedEvent().RemoveAll(this);
-	}
-
-	if (GEditor)
-	{
-		GEditor->OnLevelActorAdded().RemoveAll(this);
-		GEditor->OnLevelActorDeleted().RemoveAll(this);
-		GEditor->GetSelectedActors()->SelectionChangedEvent.RemoveAll(this);
-	}
-	
-	FEditorDelegates::MapChange.RemoveAll(this);
-	FEditorDelegates::PostUndoRedo.RemoveAll(this);
-	FCoreDelegates::OnPostEngineInit.RemoveAll(this);
-	FCoreUObjectDelegates::OnObjectTransacted.RemoveAll(this);
-
-	if (FLevelEditorModule* LevelEditorModule = FModuleManager::GetModulePtr<FLevelEditorModule>("LevelEditor"))
-	{
-		LevelEditorModule->OnComponentsEdited().RemoveAll(this);
-	}
-
-	FLevelEditorSequencerIntegration::Get().GetOnSequencersChanged().RemoveAll(this);
-
-	if (FModuleManager::Get().IsModuleLoaded(TEXT("ObjectMixer")))
-	{
-		if (UObjectMixerEditorSettings* Settings = GetMutableDefault<UObjectMixerEditorSettings>())
-		{
-			Settings->OnSettingChanged().RemoveAll(this);
-		}
-	}
-
-	for (FDelegateHandle& Delegate : DelegateHandles)
-	{
-		Delegate.Reset();
-	}
-	DelegateHandles.Empty();
-	
 	ListModel.Reset();
 
 	UToolMenus::UnregisterOwner(this);
@@ -119,7 +69,7 @@ void FObjectMixerEditorModule::Teardown()
 
 FObjectMixerEditorModule& FObjectMixerEditorModule::Get()
 {
-	return FModuleManager::LoadModuleChecked< FObjectMixerEditorModule >("ObjectMixerEditor");
+	return FModuleManager::LoadModuleChecked<FObjectMixerEditorModule>("ObjectMixerEditor");
 }
 
 void FObjectMixerEditorModule::OpenProjectSettings()
@@ -179,22 +129,6 @@ bool FObjectMixerEditorModule::RegenerateListWidget()
 	return false;
 }
 
-void FObjectMixerEditorModule::RequestRebuildList() const
-{
-	if (ListModel.IsValid())
-	{
-		ListModel->RequestRebuildList();
-	}
-}
-
-void FObjectMixerEditorModule::RefreshList() const
-{
-	if (ListModel.IsValid())
-	{
-		ListModel->RefreshList();
-	}
-}
-
 void FObjectMixerEditorModule::OnRenameCommand()
 {
 	if (ListModel.IsValid())
@@ -318,112 +252,6 @@ TSharedPtr<FWorkspaceItem> FObjectMixerEditorModule::GetWorkspaceGroup()
 	return WorkspaceGroup;
 }
 
-void FObjectMixerEditorModule::BindDelegates()
-{
-	check(GEngine && GEditor);
-	DelegateHandles.Add(GEngine->OnLevelActorListChanged().AddLambda([this] ()
-	{
-		RequestRebuildList();
-	}));
-	DelegateHandles.Add(GEngine->OnActorFolderAdded().AddLambda([this] (UActorFolder*)
-	{
-		RequestRebuildList();
-	}));
-	DelegateHandles.Add(GEngine->OnActorFoldersUpdatedEvent().AddLambda([this] (ULevel*)
-	{
-		RequestRebuildList();
-	}));
-	DelegateHandles.Add(GEngine->OnLevelActorFolderChanged().AddLambda([this] (const AActor*, FName)
-	{
-		RequestRebuildList();
-	}));
-	DelegateHandles.Add(GEngine->OnLevelActorAttached().AddLambda([this] (AActor*, const AActor*)
-	{
-		RequestRebuildList();
-	}));
-	DelegateHandles.Add(GEngine->OnLevelActorDetached().AddLambda([this] (AActor*, const AActor*)
-	{
-		RequestRebuildList();
-	}));
-
-	DelegateHandles.Add(GEditor->OnLevelActorAdded().AddLambda([this] (AActor*)
-	{
-		RequestRebuildList();
-	}));
-	DelegateHandles.Add(GEditor->OnLevelActorDeleted().AddLambda([this] (AActor*)
-	{
-		RequestRebuildList();
-	}));
-	
-	DelegateHandles.Add(FEditorDelegates::MapChange.AddLambda([this](uint32)
-	{
-		RequestRebuildList();
-	}));
-	
-	DelegateHandles.Add(FEditorDelegates::PostUndoRedo.AddLambda([this]()
-	{
-		RequestRebuildList();
-		
-		// Because we can undo/redo collection adds and removes, we need to save the data after each
-		GetMutableDefault<UObjectMixerEditorSerializedData>()->SaveConfig();
-	}));
-	
-	DelegateHandles.Add(FCoreUObjectDelegates::OnObjectTransacted.AddLambda([this](UObject*, const FTransactionObjectEvent& Event)
-	{		
-		if (Event.GetEventType() == ETransactionObjectEventType::Finalized)
-		{
-			if (Event.HasNameChange())
-			{
-				RequestRebuildList();
-			}
-			else if (const TSet<FName> PropertiesThatRequireRefresh = GetPropertiesThatRequireRefresh(); PropertiesThatRequireRefresh.Num() > 0)
-			{
-				const TSet<FName> ChangedPropertyNames = TSet<FName>(Event.GetChangedProperties());
-	
-				if (ChangedPropertyNames.Intersect(PropertiesThatRequireRefresh).Num() > 0)
-				{
-					RequestRebuildList();
-				}
-			}
-		}
-	}));
-	
-	FLevelEditorModule& LevelEditorModule = FModuleManager::LoadModuleChecked<FLevelEditorModule>("LevelEditor");
-	DelegateHandles.Add(LevelEditorModule.OnComponentsEdited().AddLambda([this]()
-	{
-		RequestRebuildList();
-	}));
-
-	DelegateHandles.Add(FLevelEditorSequencerIntegration::Get().GetOnSequencersChanged().AddLambda([this]
-	{
-		RegenerateListWidget();
-	}));
-
-	DelegateHandles.Add(GetMutableDefault<UObjectMixerEditorSettings>()->OnSettingChanged().AddLambda(
-		[this](UObject*, FPropertyChangedEvent& Event)
-		{
-			if (UObjectMixerEditorSettings::DoesPropertyChangeRequireListRebuild(Event))
-			{
-				RequestRebuildList();
-			}
-		}));
-}
-
-TSet<FName> FObjectMixerEditorModule::GetPropertiesThatRequireRefresh() const
-{
-	TSet<FName> ReturnValue;
-
-	if (ListModel.IsValid())
-	{
-		for (const TObjectPtr<UObjectMixerObjectFilter>& Instance : ListModel->GetObjectFilterInstances())
-		{
-			ReturnValue.Append(Instance->GetPropertiesThatRequireListRefresh());
-		}
-	}
-
-	return ReturnValue;
-}
-
 const TSubclassOf<UObjectMixerObjectFilter>& FObjectMixerEditorModule::GetDefaultFilterClass() const
 {
 	return DefaultFilterClass;
diff --git a/Engine/Plugins/Editor/ObjectMixer/ObjectMixer/Source/ObjectMixer/Private/SelectionInterface/LevelEditorObjectMixerSelectionInterface.cpp b/Engine/Plugins/Editor/ObjectMixer/ObjectMixer/Source/ObjectMixer/Private/SelectionInterface/LevelEditorObjectMixerSelectionInterface.cpp
new file mode 100644
--- /dev/null
+++ b/Engine/Plugins/Editor/ObjectMixer/ObjectMixer/Source/ObjectMixer/Private/SelectionInterface/LevelEditorObjectMixerSelectionInterface.cpp
@@ -0,0 +1,88 @@
+// Copyright Epic Games, Inc. All Rights Reserved.
+
+#include "SelectionInterface/LevelEditorObjectMixerSelectionInterface.h"
+
+#include "Algo/Transform.h"
+#include "Editor.h"
+#include "Engine/Selection.h"
+
+FLevelEditorObjectMixerSelectionInterface::FLevelEditorObjectMixerSelectionInterface()
+{
+	USelection::SelectionChangedEvent.AddRaw(this, &FLevelEditorObjectMixerSelectionInterface::OnLevelSelectionChanged);
+	USelection::SelectObjectEvent.AddRaw(this, &FLevelEditorObjectMixerSelectionInterface::OnLevelSelectionChanged);
+}
+
+FLevelEditorObjectMixerSelectionInterface::~FLevelEditorObjectMixerSelectionInterface()
+{
+	USelection::SelectionChangedEvent.RemoveAll(this);
+	USelection::SelectObjectEvent.RemoveAll(this);
+}
+
+void FLevelEditorObjectMixerSelectionInterface::SelectActors(const TArray<AActor*>& InSelectedActors, bool bShouldSelect, bool bSelectEvenIfHidden)
+{
+	GEditor->GetSelectedActors()->Modify();
+
+	// We'll batch selection changes instead by using BeginBatchSelectOperation()
+	GEditor->GetSelectedActors()->BeginBatchSelectOperation();
+
+	if (bShouldSelect)
+	{
+		// Clear the selection
+		GEditor->GetSelectedActors()->DeselectAll();
+	}
+
+	for (AActor* Actor : InSelectedActors)
+	{
+		constexpr bool bNotifyAfterSelect = false;
+		GEditor->SelectActor(Actor, bShouldSelect, bNotifyAfterSelect, bSelectEvenIfHidden);
+	}
+
+	// Commit selection changes
+	GEditor->GetSelectedActors()->EndBatchSelectOperation(/*bNotify*/false);
+
+	GEditor->NoteSelectionChange();
+}
+
+void FLevelEditorObjectMixerSelectionInterface::SelectComponents(const TArray<UActorComponent*>& InSelectedComponents, bool bShouldSelect, bool bSelectEvenIfHidden)
+{
+	GEditor->GetSelectedComponents()->Modify();
+
+	// We'll batch selection changes instead by using BeginBatchSelectOperation()
+	GEditor->GetSelectedComponents()->BeginBatchSelectOperation();
+
+	if (bShouldSelect)
+	{
+		// Clear the selection
+		GEditor->GetSelectedComponents()->DeselectAll();
+	}
+
+	for (UActorComponent* Component : InSelectedComponents)
+	{
+		constexpr bool bNotifyAfterSelect = false;
+		GEditor->SelectComponent(Component, bShouldSelect, bNotifyAfterSelect, bSelectEvenIfHidden);
+	}
+
+	// Commit selection changes
+	GEditor->GetSelectedComponents()->EndBatchSelectOperation(/*bNotify*/false);
+
+	GEditor->NoteSelectionChange();
+}
+
+TArray<AActor*> FLevelEditorObjectMixerSelectionInterface::GetSelectedActors() const
+{
+	TArray<AActor*> OutActors;
+	GEditor->GetSelectedActors()->GetSelectedObjects<AActor>(OutActors);
+	return MoveTemp(OutActors);
+}
+
+TArray<UActorComponent*> FLevelEditorObjectMixerSelectionInterface::GetSelectedComponents() const
+{
+	TArray<UActorComponent*> OutComponents;
+	GEditor->GetSelectedComponents()->GetSelectedObjects<UActorComponent>(OutComponents);
+	return MoveTemp(OutComponents);
+}
+
+void FLevelEditorObjectMixerSelectionInterface::OnLevelSelectionChanged(UObject* Obj)
+{
+	SelectionChanged.Broadcast();
+}
diff --git a/Engine/Plugins/Editor/ObjectMixer/ObjectMixer/Source/ObjectMixer/Private/Views/List/Modes/ObjectMixerOutlinerHierarchy.cpp b/Engine/Plugins/Editor/ObjectMixer/ObjectMixer/Source/ObjectMixer/Private/Views/List/Modes/ObjectMixerOutlinerHierarchy.cpp
--- a/Engine/Plugins/Editor/ObjectMixer/ObjectMixer/Source/ObjectMixer/Private/Views/List/Modes/ObjectMixerOutlinerHierarchy.cpp
+++ b/Engine/Plugins/Editor/ObjectMixer/ObjectMixer/Source/ObjectMixer/Private/Views/List/Modes/ObjectMixerOutlinerHierarchy.cpp
@@ -607,7 +607,7 @@ TArray<FSceneOutlinerTreeItemPtr> FObjectMixerOutlinerHierarchy::ConditionallyCr
 			Mode->CreateItemFor<FObjectMixerEditorListRowActor>(
 				FObjectMixerEditorListRowActor(Actor, GetCastedMode()->GetSceneOutliner())))
 		{
-			const bool bIsHybridModeEnabled = GetDefault<UObjectMixerEditorSettings>()->IsHybridRowModeEnabled();
+			const bool bIsHybridModeEnabled = GetCastedMode()->ShouldAllowHybridRows() && GetDefault<UObjectMixerEditorSettings>()->IsHybridRowModeEnabled();
 			if (bIsHybridModeEnabled && ComponentRows.Num() == 1) // Create hybrid row
 			{
 				if (FObjectMixerEditorListRowActor* AsActorRow = FObjectMixerUtils::AsActorRow(ActorItem))
diff --git a/Engine/Plugins/Editor/ObjectMixer/ObjectMixer/Source/ObjectMixer/Private/Views/List/Modes/ObjectMixerOutlinerMode.cpp b/Engine/Plugins/Editor/ObjectMixer/ObjectMixer/Source/ObjectMixer/Private/Views/List/Modes/ObjectMixerOutlinerMode.cpp
--- a/Engine/Plugins/Editor/ObjectMixer/ObjectMixer/Source/ObjectMixer/Private/Views/List/Modes/ObjectMixerOutlinerMode.cpp
+++ b/Engine/Plugins/Editor/ObjectMixer/ObjectMixer/Source/ObjectMixer/Private/Views/List/Modes/ObjectMixerOutlinerMode.cpp
@@ -5,10 +5,12 @@
 #include "ObjectMixerEditorLog.h"
 #include "ObjectMixerEditorModule.h"
 #include "ObjectMixerEditorSettings.h"
+#include "ObjectFilter/ObjectMixerEditorObjectFilter.h"
 #include "Views/List/ObjectMixerEditorList.h"
 #include "Views/List/Modes/ObjectMixerOutlinerHierarchy.h"
 #include "Views/List/Modes/SFilterClassMenuItem.h"
 #include "Views/List/ObjectMixerUtils.h"
+#include "Views/List/SObjectMixerEditorList.h"
 #include "Views/List/RowTypes/ObjectMixerEditorListRowActor.h"
 #include "Views/Widgets/ObjectMixerEditorListMenuContext.h"
 
@@ -79,6 +81,12 @@ TObjectPtr<UObjectMixerOutlinerModeEditorConfig> UObjectMixerOutlinerModeEditorC
 
 namespace ObjectMixerOutliner
 {
+	FHybridRowSelector::FHybridRowSelector(const FObjectMixerOutlinerMode* Mode)
+	{
+		check(Mode);
+		bAllowHybridRows = Mode->ShouldAllowHybridRows();
+	}
+
 	bool FWeakActorSelectorAcceptingComponents::operator()(const TWeakPtr<ISceneOutlinerTreeItem>& Item, TWeakObjectPtr<AActor>& DataOut) const
 	{
 		if (TSharedPtr<ISceneOutlinerTreeItem> ItemPtr = Item.Pin())
@@ -125,7 +133,7 @@ namespace ObjectMixerOutliner
 				if (ActorItem->IsValid())
 				{
 					// Get the component from a hybrid row if we're in HybridComponent mode
-					if (GetDefault<UObjectMixerEditorSettings>()->HybridRowPolicy == EObjectMixerHybridMode::HybridComponent &&
+					if (ShouldAllowHybridRows() && GetDefault<UObjectMixerEditorSettings>()->HybridRowPolicy == EObjectMixerHybridMode::HybridComponent &&
 						ActorItem->RowData.GetHybridComponent())
 					{
 						DataOut = ActorItem->RowData.GetHybridComponent();
@@ -163,7 +171,7 @@ namespace ObjectMixerOutliner
 				if (ActorItem->IsValid())
 				{
 					// Skip hybrid rows here if in HybridComponent mode, we'll get the components from them in FComponentSelector
-					if (GetDefault<UObjectMixerEditorSettings>()->HybridRowPolicy == EObjectMixerHybridMode::HybridComponent &&
+					if (ShouldAllowHybridRows() && GetDefault<UObjectMixerEditorSettings>()->HybridRowPolicy == EObjectMixerHybridMode::HybridComponent &&
 						ActorItem->RowData.GetIsHybridRow())
 					{
 						return false;
@@ -500,8 +508,13 @@ FObjectMixerOutlinerMode::FObjectMixerOutlinerMode(
 		return false;
 	}), FSceneOutlinerFilter::EDefaultBehaviour::Pass));
 
-	USelection::SelectionChangedEvent.AddRaw(this, &FObjectMixerOutlinerMode::OnLevelSelectionChanged);
-	USelection::SelectObjectEvent.AddRaw(this, &FObjectMixerOutlinerMode::OnLevelSelectionChanged);
+	if (TSharedPtr<FObjectMixerEditorList> PinnedList = GetListModelPtr().Pin())
+	{
+		if (PinnedList->GetSelectionInterface().IsValid())
+		{
+			PinnedList->GetSelectionInterface()->OnSelectionChanged().AddRaw(this, &FObjectMixerOutlinerMode::OnEditorSelectionChanged);
+		}
+	}
 
 	FEditorDelegates::MapChange.AddRaw(this, &FObjectMixerOutlinerMode::OnMapChange);
 	FEditorDelegates::NewCurrentLevel.AddRaw(this, &FObjectMixerOutlinerMode::OnNewCurrentLevel);
@@ -541,6 +554,14 @@ FObjectMixerOutlinerMode::~FObjectMixerOutlinerMode()
 		}
 	}
 
+	if (TSharedPtr<FObjectMixerEditorList> PinnedList = GetListModelPtr().Pin())
+	{
+		if (PinnedList->GetSelectionInterface().IsValid())
+		{
+			PinnedList->GetSelectionInterface()->OnSelectionChanged().RemoveAll(this);
+		}
+	}
+
 	FEditorDelegates::OnEditCutActorsBegin.RemoveAll(this);
 	FEditorDelegates::OnEditCutActorsEnd.RemoveAll(this);
 	FEditorDelegates::OnEditCopyActorsBegin.RemoveAll(this);
@@ -552,9 +573,6 @@ FObjectMixerOutlinerMode::~FObjectMixerOutlinerMode()
 	FEditorDelegates::OnDeleteActorsBegin.RemoveAll(this);
 	FEditorDelegates::OnDeleteActorsEnd.RemoveAll(this);
 
-	USelection::SelectionChangedEvent.RemoveAll(this);
-	USelection::SelectObjectEvent.RemoveAll(this);
-
 	FEditorDelegates::MapChange.RemoveAll(this);
 	FEditorDelegates::NewCurrentLevel.RemoveAll(this);
 
@@ -575,8 +593,13 @@ void FObjectMixerOutlinerMode::OnNewCurrentLevel()
 	SceneOutliner->FullRefresh();
 }
 
-void FObjectMixerOutlinerMode::OnLevelSelectionChanged(UObject* Obj)
+void FObjectMixerOutlinerMode::OnEditorSelectionChanged()
 {
+	if (!ShouldSyncSelectionFromEditor())
+	{
+		return;
+	}
+
 	const FSceneOutlinerFilterInfo* ShowOnlySelectedActorsFilter = FilterInfoMap.Find(TEXT("ShowOnlySelectedActors"));
 
 	// Since there is no way to know which items were removed/added to a selection, we must force a full refresh to handle this
@@ -592,22 +615,25 @@ void FObjectMixerOutlinerMode::OnLevelSelectionChanged(UObject* Obj)
 		SceneOutliner->RefreshSelection();
 
 		// Scroll last item into view - this means if we are multi-selecting, we show newest selection. @TODO Not perfect though
-		if (const AActor* LastSelectedActor = GEditor->GetSelectedActors()->GetBottom<AActor>())
+		const TArray<AActor*> SelectedActors = GetSelectedActorsInEditor();
+		if (!SelectedActors.IsEmpty())
 		{
-			if (FSceneOutlinerTreeItemPtr TreeItem = SceneOutliner->GetTreeItem(LastSelectedActor, false))
+			if (const AActor* LastSelectedActor = SelectedActors.Last())
 			{
-				// Only scroll if selection framing is enabled
-				if(bAlwaysFrameSelection)
+				if (FSceneOutlinerTreeItemPtr TreeItem = SceneOutliner->GetTreeItem(LastSelectedActor, false))
 				{
-					SceneOutliner->ScrollItemIntoView(TreeItem);
+					// Only scroll if selection framing is enabled
+					if (bAlwaysFrameSelection)
+					{
+						SceneOutliner->ScrollItemIntoView(TreeItem);
+					}
+				}
+				else
+				{
+					SceneOutliner::ENewItemAction::Type Action = bAlwaysFrameSelection ? SceneOutliner::ENewItemAction::ScrollIntoView : SceneOutliner::ENewItemAction::Select;
+
+					SceneOutliner->OnItemAdded(LastSelectedActor, Action);
 				}
-			}
-			else
-			{
-				SceneOutliner::ENewItemAction::Type Action = bAlwaysFrameSelection ? SceneOutliner::ENewItemAction::ScrollIntoView : SceneOutliner::ENewItemAction::Select;
-				
-				SceneOutliner->OnItemAdded(LastSelectedActor, Action);
-				
 			}
 		}
 	}
@@ -669,11 +695,11 @@ void FObjectMixerOutlinerMode::SynchronizeAllSelectionsToEditor()
 
 	// Actors
 	SynchronizeSelectedActorDescs();
-	TArray<AActor*> SelectedActors = SceneOutliner->GetSelection().GetData<AActor*>(ObjectMixerOutliner::FActorSelector());
+	TArray<AActor*> SelectedActors = SceneOutliner->GetSelection().GetData<AActor*>(ObjectMixerOutliner::FActorSelector(this));
 	bool bHasActorSelectionChanged = HasActorSelectionChanged(SelectedActors, bAreAnyInPIE);
 		
 	// Components
-	TArray<UActorComponent*> SelectedComponents = SceneOutliner->GetSelection().GetData<UActorComponent*>(ObjectMixerOutliner::FComponentSelector());
+	TArray<UActorComponent*> SelectedComponents = SceneOutliner->GetSelection().GetData<UActorComponent*>(ObjectMixerOutliner::FComponentSelector(this));
 	bool bHasComponentSelectionChanged = HasComponentSelectionChanged(SelectedComponents, bAreAnyInPIE);
 		
 	// If there's a discrepancy, update the selected objects to reflect the list.
@@ -684,18 +710,20 @@ void FObjectMixerOutlinerMode::SynchronizeAllSelectionsToEditor()
 
 		if (bHasActorSelectionChanged)
 		{
-			SelectActorsInEditor(SelectedActors);
+			SelectActorsInEditor(SelectedActors, true, true);
 		}
 			
 		if (bHasComponentSelectionChanged)
 		{
-			SelectComponentsInEditor(SelectedComponents);
+			SelectComponentsInEditor(SelectedComponents, true, true);
 		}
 	}
 }
 
 bool FObjectMixerOutlinerMode::HasActorSelectionChanged(TArray<AActor*>& OutSelectedActors, bool& bOutAreAnyInPIE)
 {
+	TArray<AActor*> EditorSelectedActors = GetSelectedActorsInEditor();
+
 	bool bHasSelectionChanged = false;
 	for (AActor* Actor : OutSelectedActors)
 	{
@@ -703,16 +731,15 @@ bool FObjectMixerOutlinerMode::HasActorSelectionChanged(TArray<AActor*>& OutSele
 		{
 			bOutAreAnyInPIE = true;
 		}
-		if (!GEditor->GetSelectedActors()->IsSelected(Actor))
+		if (!EditorSelectedActors.Contains(Actor))
 		{
 			bHasSelectionChanged = true;
 			break;
 		}
 	}
 
-	for (FSelectionIterator SelectionIt(*GEditor->GetSelectedActors()); SelectionIt && !bHasSelectionChanged; ++SelectionIt)
+	for (AActor* Actor : EditorSelectedActors)
 	{
-		const AActor* Actor = CastChecked< AActor >(*SelectionIt);
 		if (!bOutAreAnyInPIE && Actor->GetPackage()->HasAnyPackageFlags(PKG_PlayInEditor))
 		{
 			bOutAreAnyInPIE = true;
@@ -735,6 +762,11 @@ bool FObjectMixerOutlinerMode::HasActorSelectionChanged(TArray<AActor*>& OutSele
 				}
 			}
 		}
+
+		if (bHasSelectionChanged)
+		{
+			break;
+		}
 	}
 
 	return bHasSelectionChanged;
@@ -742,6 +774,8 @@ bool FObjectMixerOutlinerMode::HasActorSelectionChanged(TArray<AActor*>& OutSele
 
 bool FObjectMixerOutlinerMode::HasComponentSelectionChanged(TArray<UActorComponent*>& OutSelectedComponents, bool& bOutAreAnyInPIE)
 {
+	TArray<UActorComponent*> EditorSelectedComponents = GetSelectedComponentsInEditor();
+
 	bool bHasSelectionChanged = false;
 	for (UActorComponent* Component : OutSelectedComponents)
 	{
@@ -749,72 +783,124 @@ bool FObjectMixerOutlinerMode::HasComponentSelectionChanged(TArray<UActorCompone
 		{
 			bOutAreAnyInPIE = true;
 		}
-		if (!GEditor->GetSelectedComponents()->IsSelected(Component))
+		if (!EditorSelectedComponents.Contains(Component))
 		{
 			bHasSelectionChanged = true;
 			break;
 		}
 	}
 
-	for (FSelectionIterator SelectionIt(*GEditor->GetSelectedComponents()); SelectionIt && !bHasSelectionChanged; ++SelectionIt)
+	for (UActorComponent* Component : EditorSelectedComponents)
 	{
-		const UActorComponent* Actor = CastChecked< UActorComponent >(*SelectionIt);
-		if (!bOutAreAnyInPIE && Actor->GetPackage()->HasAnyPackageFlags(PKG_PlayInEditor))
+		if (!bOutAreAnyInPIE && Component->GetPackage()->HasAnyPackageFlags(PKG_PlayInEditor))
 		{
 			bOutAreAnyInPIE = true;
 		}
-		if (!OutSelectedComponents.Contains(Actor))
+		if (!OutSelectedComponents.Contains(Component))
 		{
-			// Actor has been deselected
+			// Component has been deselected
 			bHasSelectionChanged = true;
+			break;
 		}
 	}
 
 	return bHasSelectionChanged;
 }
 
-void FObjectMixerOutlinerMode::SelectActorsInEditor(const TArray<AActor*>& InSelectedActors)
+void FObjectMixerOutlinerMode::SelectActorsInEditor(const TArray<AActor*>& InSelectedActors, bool bShouldSelect, bool bSelectEvenIfHidden)
 {
-	GEditor->GetSelectedActors()->Modify();
-				
-	// We'll batch selection changes instead by using BeginBatchSelectOperation()
-	GEditor->GetSelectedActors()->BeginBatchSelectOperation();
+	if (TSharedPtr<FObjectMixerEditorList> PinnedList = GetListModelPtr().Pin())
+	{
+		if (PinnedList->GetSelectionInterface().IsValid())
+		{
+			PinnedList->GetSelectionInterface()->SelectActors(InSelectedActors, bShouldSelect, bSelectEvenIfHidden);
+		}
+	}
+}
 
-	// Clear the selection
-	GEditor->GetSelectedActors()->DeselectAll();
-				
-	for (AActor* Actor : InSelectedActors)
+void FObjectMixerOutlinerMode::SelectComponentsInEditor(const TArray<UActorComponent*>& InSelectedComponents, bool bShouldSelect, bool bSelectEvenIfHidden)
+{
+	if (TSharedPtr<FObjectMixerEditorList> PinnedList = GetListModelPtr().Pin())
 	{
-		constexpr bool bShouldSelect = true;
-		constexpr bool bNotifyAfterSelect = false;
-		constexpr bool bSelectEvenIfHidden = true;
-		GEditor->SelectActor(Actor, bShouldSelect, bNotifyAfterSelect, bSelectEvenIfHidden);
+		if (PinnedList->GetSelectionInterface().IsValid())
+		{
+			PinnedList->GetSelectionInterface()->SelectComponents(InSelectedComponents, bShouldSelect, bSelectEvenIfHidden);
+		}
 	}
+}
+
+void FObjectMixerOutlinerMode::SelectActorsInMixer(const TArray<AActor*>& InSelectedActors, bool bShouldSelect, bool bSelectEvenIfHidden)
+{
+	TArray<FSceneOutlinerTreeItemPtr> ItemsToSelect;
+	ItemsToSelect.Reserve(InSelectedActors.Num());
 
-	// Commit selection changes
-	GEditor->GetSelectedActors()->EndBatchSelectOperation(/*bNotify*/false);
+	Algo::TransformIf(InSelectedActors, ItemsToSelect,
+		[this, bSelectEvenIfHidden](AActor* Actor) -> bool
+		{
+			if (const FSceneOutlinerTreeItemPtr& Item = SceneOutliner->GetTreeItem(Actor))
+			{
+				return bSelectEvenIfHidden || !Item->Flags.bIsFilteredOut;
+			}
+
+			return false;
+		},
+		[this](AActor* Actor) -> FSceneOutlinerTreeItemPtr
+		{
+			return SceneOutliner->GetTreeItem(Actor);
+		}
+	);
+
+	SceneOutliner->SetItemSelection(ItemsToSelect, bShouldSelect);
 }
 
-void FObjectMixerOutlinerMode::SelectComponentsInEditor(const TArray<UActorComponent*>& InSelectedComponents)
-{	
-	GEditor->GetSelectedComponents()->Modify();
-				
-	// We'll batch selection changes instead by using BeginBatchSelectOperation()
-	GEditor->GetSelectedComponents()->BeginBatchSelectOperation();
+void FObjectMixerOutlinerMode::SelectComponentsInMixer(const TArray<UActorComponent*>& InSelectedComponents, bool bShouldSelect, bool bSelectEvenIfHidden)
+{
+	TArray<FSceneOutlinerTreeItemPtr> ItemsToSelect;
+	ItemsToSelect.Reserve(InSelectedComponents.Num());
 
-	// Clear the selection
-	GEditor->GetSelectedComponents()->DeselectAll();
-				
-	for (UActorComponent* Component : InSelectedComponents)
+	Algo::TransformIf(InSelectedComponents, ItemsToSelect,
+		[this, bSelectEvenIfHidden](UActorComponent* Component) -> bool
+		{
+			if (const FSceneOutlinerTreeItemPtr& Item = SceneOutliner->GetTreeItem(Component))
+			{
+				return bSelectEvenIfHidden || !Item->Flags.bIsFilteredOut;
+			}
+
+			return false;
+		},
+		[this](UActorComponent* Component) -> FSceneOutlinerTreeItemPtr
+		{
+			return SceneOutliner->GetTreeItem(Component);
+		}
+	);
+
+	SceneOutliner->SetItemSelection(ItemsToSelect, bShouldSelect);
+}
+
+TArray<AActor*> FObjectMixerOutlinerMode::GetSelectedActorsInEditor() const
+{
+	if (TSharedPtr<FObjectMixerEditorList> PinnedList = GetListModelPtr().Pin())
 	{
-		constexpr bool bShouldSelect = true;
-		constexpr bool bNotifyAfterSelect = false;
-		constexpr bool bSelectEvenIfHidden = true;
-		GEditor->SelectComponent(Component, bShouldSelect, bNotifyAfterSelect, bSelectEvenIfHidden);
+		if (PinnedList->GetSelectionInterface().IsValid())
+		{
+			return PinnedList->GetSelectionInterface()->GetSelectedActors();
+		}
 	}
 
-	// Commit selection changes
-	GEditor->GetSelectedComponents()->EndBatchSelectOperation(/*bNotify*/false);
+	return TArray<AActor*>();
+}
+
+TArray<UActorComponent*> FObjectMixerOutlinerMode::GetSelectedComponentsInEditor() const
+{
+	if (TSharedPtr<FObjectMixerEditorList> PinnedList = GetListModelPtr().Pin())
+	{
+		if (PinnedList->GetSelectionInterface().IsValid())
+		{
+			return PinnedList->GetSelectionInterface()->GetSelectedComponents();
+		}
+	}
+
+	return TArray<UActorComponent*>();
 }
 
 void FObjectMixerOutlinerMode::OnActorLabelChanged(AActor* ChangedActor)
@@ -1768,14 +1854,16 @@ TSharedPtr<SWidget> FObjectMixerOutlinerMode::BuildContextMenu()
 
 TSharedPtr<SWidget> FObjectMixerOutlinerMode::CreateContextMenu()
 {
+	constexpr bool bShouldSelect = false;
+	constexpr bool bSelectEvenIfHidden = true;
+
 	// Make sure that no components are selected
-	if (GEditor->GetSelectedComponentCount() > 0)
+	TArray<UActorComponent*> SelectedComponents = GetSelectedComponentsInEditor();
+	if (!SelectedComponents.IsEmpty())
 	{
 		// We want to be able to undo to regain the previous component selection
 		const FScopedTransaction Transaction(NSLOCTEXT("UnrealEd", "ClickingOnActorsContextMenu", "Clicking on Actors (context menu)"));
-		USelection* ComponentSelection = GEditor->GetSelectedComponents();
-		ComponentSelection->Modify(false);
-		ComponentSelection->DeselectAll();
+		SelectComponentsInEditor(SelectedComponents, bShouldSelect, bSelectEvenIfHidden);
 
 		GUnrealEd->UpdatePivotLocationForSelection();
 		GEditor->RedrawLevelEditingViewports(false);
@@ -1792,8 +1880,8 @@ void FObjectMixerOutlinerMode::OnItemAdded(FSceneOutlinerTreeItemPtr Item)
 		{
 			++FilteredActorCount;
 
-			// Synchronize selection
-			if (GEditor->GetSelectedActors()->IsSelected(ActorItem->Actor.Get()))
+			// Synchronize selection with editor (or always select if sync is disabled)
+			if (!ShouldSyncSelectionFromEditor() || GetSelectedActorsInEditor().Contains(ActorItem->Actor.Get()))
 			{
 				SceneOutliner->SetItemSelection(Item, true);
 			}
@@ -1910,7 +1998,7 @@ void FObjectMixerOutlinerMode::OnItemDoubleClick(FSceneOutlinerTreeItemPtr Item)
 			if (Selection.Has<FActorTreeItem>())
 			{
 				const bool bActiveViewportOnly = false;
-				GEditor->MoveViewportCamerasToActor(Selection.GetData<AActor*>(ObjectMixerOutliner::FActorSelector()), bActiveViewportOnly);
+				GEditor->MoveViewportCamerasToActor(Selection.GetData<AActor*>(ObjectMixerOutliner::FActorSelector(this)), bActiveViewportOnly);
 			}
 		}
 		else
@@ -1927,26 +2015,18 @@ void FObjectMixerOutlinerMode::OnItemDoubleClick(FSceneOutlinerTreeItemPtr Item)
 
 void FObjectMixerOutlinerMode::OnFilterTextCommited(FSceneOutlinerItemSelection& Selection, ETextCommit::Type CommitType)
 {
-	// Start batching selection changes
-	GEditor->GetSelectedActors()->BeginBatchSelectOperation();
+	const TArray<AActor*> Actors = Selection.GetData<AActor*>(ObjectMixerOutliner::FActorSelector(this));
 
-	// Select actors (and only the actors) that match the filter text
-	const bool bNoteSelectionChange = false;
-	const bool bDeselectBSPSurfs = false;
-	const bool WarnAboutManyActors = true;
-	GEditor->SelectNone(bNoteSelectionChange, bDeselectBSPSurfs, WarnAboutManyActors);
-	for (AActor* Actor : Selection.GetData<AActor*>(ObjectMixerOutliner::FActorSelector()))
+	constexpr bool bShouldSelect = true;
+	constexpr bool bSelectEvenIfHidden = false;
+	if (ShouldSyncSelectionToEditor())
 	{
-		const bool bShouldSelect = true;
-		const bool bSelectEvenIfHidden = false;
-		GEditor->SelectActor(Actor, bShouldSelect, bNoteSelectionChange, bSelectEvenIfHidden);
+		SelectActorsInEditor(Actors, bShouldSelect, bSelectEvenIfHidden);
+	}
+	else
+	{
+		SelectActorsInMixer(Actors, bShouldSelect, bSelectEvenIfHidden);
 	}
-
-	// Commit selection changes
-	GEditor->GetSelectedActors()->EndBatchSelectOperation(/*bNotify*/false);
-
-	// Fire selection changed event
-	GEditor->NoteSelectionChange();
 
 	// Set keyboard focus to the SceneOutliner, so the user can perform keyboard commands that interact
 	// with selected actors (such as Delete, to delete selected actors.)
@@ -2276,28 +2356,53 @@ FFolder FObjectMixerOutlinerMode::GetWorldDefaultRootFolder() const
 	return FFolder::GetWorldRootFolder(RepresentingWorld.Get());
 }
 
-void FObjectMixerOutlinerMode::SynchronizeComponentSelection()
+void FObjectMixerOutlinerMode::SynchronizeComponentAndActorSelection()
 {
-	USelection* SelectedComponents = GEditor->GetSelectedComponents();
+	// Implementation based FActorMode::SynchronizeActorSelection, but we use GetSelectedActorsInEditor instead and handle components as well
 
-	// Deselect components in the tree that are no longer selected in the world
+	TArray<UActorComponent*> SelectedComponents = GetSelectedComponentsInEditor();
+
+	// Deselect components in the tree that are no longer selected in the editor.
 	const FSceneOutlinerItemSelection Selection(SceneOutliner->GetSelection());
-	auto DeselectComponents = [this](FComponentTreeItem& Item)
+	auto DeselectComponents = [this, &SelectedComponents](FComponentTreeItem& Item)
 	{
-		if (!Item.Component.IsValid() || !Item.Component.Get()->IsSelected())
+		if (!Item.Component.IsValid() || !SelectedComponents.Contains(Item.Component.Get()))
 		{
 			SceneOutliner->SetItemSelection(Item.AsShared(), false);
+			return;
 		}
 	};
 	Selection.ForEachItem<FComponentTreeItem>(DeselectComponents);
 
+	TArray<AActor*> SelectedActors = GetSelectedActorsInEditor();
+
+	// Remove any actors from the selected list that own selected components, since they're necessarily selected when an owned component is
+	// selected, but presumably not the user's editing target.
+	// NOTE: this assumption needs to be revisited if we implement multi-select
+	for (UActorComponent* Component : SelectedComponents)
+	{
+		if (Component)
+		{
+			SelectedActors.Remove(Component->GetOwner());
+		}
+	}
+	
+	// Deselect actors in the tree that are no longer selected in the editor (or owned by a selected component)
+	auto DeselectActors = [this, &SelectedActors](FActorTreeItem& Item)
+	{
+		if (!Item.Actor.IsValid() || !SelectedActors.Contains(Item.Actor.Get()))
+		{
+			SceneOutliner->SetItemSelection(Item.AsShared(), false);
+		}
+	};
+	Selection.ForEachItem<FActorTreeItem>(DeselectActors);
+
 	// See if the tree view selector is pointing at a selected item
 	bool bSelectorInSelectionSet = false;
 
-	TArray<FSceneOutlinerTreeItemPtr> ComponentItems;
-	for (FSelectionIterator SelectionIt(*SelectedComponents); SelectionIt; ++SelectionIt)
+	TArray<FSceneOutlinerTreeItemPtr> SelectedItems;
+	for (UActorComponent* Component : SelectedComponents)
 	{
-		UActorComponent* Component = CastChecked< UActorComponent >(*SelectionIt);
 		if (FSceneOutlinerTreeItemPtr ComponentItem = SceneOutliner->GetTreeItem(Component))
 		{
 			if (!bSelectorInSelectionSet && SceneOutliner->HasSelectorFocus(ComponentItem))
@@ -2305,13 +2410,26 @@ void FObjectMixerOutlinerMode::SynchronizeComponentSelection()
 				bSelectorInSelectionSet = true;
 			}
 
-			ComponentItems.Add(ComponentItem);
+			SelectedItems.Add(ComponentItem);
+		}
+	}
+
+	for (AActor* Actor : SelectedActors)
+	{
+		if (FSceneOutlinerTreeItemPtr ActorItem = SceneOutliner->GetTreeItem(Actor))
+		{
+			if (!bSelectorInSelectionSet && SceneOutliner->HasSelectorFocus(ActorItem))
+			{
+				bSelectorInSelectionSet = true;
+			}
+
+			SelectedItems.Add(ActorItem);
 		}
 	}
 
 	// If NOT bSelectorInSelectionSet then we want to just move the selector to the first selected item.
 	ESelectInfo::Type SelectInfo = bSelectorInSelectionSet ? ESelectInfo::Direct : ESelectInfo::OnMouseClick;
-	SceneOutliner->AddToSelection(ComponentItems, SelectInfo);
+	SceneOutliner->AddToSelection(SelectedItems, SelectInfo);
 
 	FSceneOutlinerDelegates::Get().SelectionChanged.Broadcast();
 }
@@ -2677,7 +2795,7 @@ void FObjectMixerOutlinerMode::OnDrop(ISceneOutlinerTreeItem& DropTarget, const
 				Payload.ForEachItem<FFolderTreeItem>(MoveToDestination);
 
 				// Since target root is directly the Level Instance, clear folder path
-				TArray<AActor*> DraggedActors = Payload.GetData<AActor*>(ObjectMixerOutliner::FActorSelector());
+				TArray<AActor*> DraggedActors = Payload.GetData<AActor*>(ObjectMixerOutliner::FActorSelector(this));
 				for (auto& Actor : DraggedActors)
 				{
 					Actor->SetFolderPath_Recursively(FName());
@@ -2942,7 +3060,9 @@ namespace ObjectMixerActorBrowsingModeUtils
 		}
 	}
 
-	static void RecursiveActorSelect(SSceneOutliner* SceneOutliner, const FSceneOutlinerTreeItemPtr& Item, bool bSelectImmediateChildrenOnly)
+	// Get a list of unique actors to select as descendents of the Item in the SceneOutliner
+	// Note that this also selects any folders encountered during the search
+	static void RecursiveGetActorsToSelect(SSceneOutliner* SceneOutliner, const FSceneOutlinerTreeItemPtr& Item, bool bSelectImmediateChildrenOnly, TArray<AActor*>& OutActorsToSelect)
 	{
 		if (Item.IsValid())
 		{
@@ -2951,7 +3071,7 @@ namespace ObjectMixerActorBrowsingModeUtils
 			{
 				if (AActor* Actor = ActorItem->Actor.Get())
 				{
-					GEditor->SelectActor(Actor, true, false);
+					OutActorsToSelect.AddUnique(Actor);
 				}
 			}
 			// Select all children
@@ -2964,7 +3084,7 @@ namespace ObjectMixerActorBrowsingModeUtils
 					{
 						if (AActor* Actor = ActorItem->Actor.Get())
 						{
-							GEditor->SelectActor(Actor, true, false);
+							OutActorsToSelect.AddUnique(Actor);
 						}
 					}
 					else if (FFolderTreeItem* FolderItem = ChildPtr->CastTo<FFolderTreeItem>())
@@ -2976,7 +3096,7 @@ namespace ObjectMixerActorBrowsingModeUtils
 					{
 						for (const TWeakPtr<ISceneOutlinerTreeItem>& Grandchild : ChildPtr->GetChildren())
 						{
-							RecursiveActorSelect(SceneOutliner, Grandchild.Pin(), bSelectImmediateChildrenOnly);
+							RecursiveGetActorsToSelect(SceneOutliner, Grandchild.Pin(), bSelectImmediateChildrenOnly, OutActorsToSelect);
 						}
 					}
 				}
@@ -3071,15 +3191,22 @@ void FObjectMixerOutlinerMode::SelectFoldersDescendants(const TArray<FFolderTree
 	}
 
 	// batch selection
-	GEditor->GetSelectedActors()->BeginBatchSelectOperation();
-
+	TArray<AActor*> ActorsToSelect;
 	for (FFolderTreeItem* Folder : FolderItems)
 	{
-		ObjectMixerActorBrowsingModeUtils::RecursiveActorSelect(SceneOutliner, Folder->AsShared(), bSelectImmediateChildrenOnly);
+		ObjectMixerActorBrowsingModeUtils::RecursiveGetActorsToSelect(SceneOutliner, Folder->AsShared(), bSelectImmediateChildrenOnly, ActorsToSelect);
 	}
 
-	GEditor->GetSelectedActors()->EndBatchSelectOperation(/*bNotify*/false);
-	GEditor->NoteSelectionChange();
+	constexpr bool bShouldSelect = true;
+	constexpr bool bSelectEvenIfHidden = true;
+	if (ShouldSyncSelectionToEditor())
+	{
+		SelectActorsInEditor(ActorsToSelect, bShouldSelect, bSelectEvenIfHidden);
+	}
+	else
+	{
+		SelectActorsInMixer(ActorsToSelect, bShouldSelect, bSelectEvenIfHidden);
+	}
 }
 
 bool FObjectMixerOutlinerMode::CanPinItems(const TArray<FSceneOutlinerTreeItemPtr>& InItems) const
@@ -3095,32 +3222,34 @@ void FObjectMixerOutlinerMode::PinItems(const TArray<FSceneOutlinerTreeItemPtr>&
 		return;
 	}
 
-	TArray<FGuid> ActorsToPin;
+	TArray<FGuid> ActorsGuidsToPin;
 	// If Unloaded actors are hidden and we are pinning folders we need to find them through FActorFolders
 	const bool bSearchForHiddenUnloadedActors = bHideUnloadedActors;
-	ObjectMixerActorBrowsingModeUtils::RecursiveAddItemsToActorGuidList(InItems, ActorsToPin, bSearchForHiddenUnloadedActors);
+	ObjectMixerActorBrowsingModeUtils::RecursiveAddItemsToActorGuidList(InItems, ActorsGuidsToPin, bSearchForHiddenUnloadedActors);
 
-	if (ActorsToPin.Num())
+	if (ActorsGuidsToPin.Num())
 	{
-		GEditor->GetSelectedActors()->BeginBatchSelectOperation();
-		GEditor->SelectNone(/*bNoteSelectionChange=*/false, /*bDeselectBSPSurfs=*/true);
+		WorldPartition->PinActors(ActorsGuidsToPin);
 
-		WorldPartition->PinActors(ActorsToPin);
+		TArray<AActor*> ActorsToPin;
+		ActorsToPin.Reserve(ActorsGuidsToPin.Num());
 
 		AActor* LastPinnedActor = nullptr;
-		for (const FGuid& ActorGuid : ActorsToPin)
+		for (const FGuid& ActorGuid : ActorsGuidsToPin)
 		{
 			if (FWorldPartitionHandle ActorHandle(WorldPartition, ActorGuid); ActorHandle.IsValid())
 			{
 				if (AActor* PinnedActor = ActorHandle.GetActor())
 				{
-					GEditor->SelectActor(PinnedActor, /*bInSelected=*/true, /*bNotify=*/false);
+					ActorsToPin.Add(PinnedActor);
 					LastPinnedActor = PinnedActor;
 				}
 			}
 		}
 
-		GEditor->GetSelectedActors()->EndBatchSelectOperation(/*bNotify=*/true);
+		constexpr bool bShouldSelect = true;
+		constexpr bool bSelectEvenIfHidden = true;
+		SelectActorsInEditor(ActorsToPin, bShouldSelect, bSelectEvenIfHidden);
 
 		if (LastPinnedActor)
 		{
@@ -3142,29 +3271,32 @@ void FObjectMixerOutlinerMode::UnpinItems(const TArray<FSceneOutlinerTreeItemPtr
 		return;
 	}
 
-	TArray<FGuid> ActorsToUnpin;
+	TArray<FGuid> ActorGuidsToUnpin;
 	// No need to search for hidden unloaded actors when unloading
 	const bool bSearchForHiddenUnloadedActors = false;
-	ObjectMixerActorBrowsingModeUtils::RecursiveAddItemsToActorGuidList(InItems, ActorsToUnpin, bSearchForHiddenUnloadedActors);
+	ObjectMixerActorBrowsingModeUtils::RecursiveAddItemsToActorGuidList(InItems, ActorGuidsToUnpin, bSearchForHiddenUnloadedActors);
 
-	if (ActorsToUnpin.Num())
+	if (ActorGuidsToUnpin.Num())
 	{
-		GEditor->GetSelectedActors()->BeginBatchSelectOperation();
+		TArray<AActor*> ActorsToUnpin;
+		ActorsToUnpin.Reserve(ActorGuidsToUnpin.Num());
 
-		for (const FGuid& ActorGuid : ActorsToUnpin)
+		for (const FGuid& ActorGuid : ActorGuidsToUnpin)
 		{
 			if (FWorldPartitionHandle ActorHandle(WorldPartition, ActorGuid); ActorHandle.IsValid())
 			{
 				if (AActor* PinnedActor = ActorHandle.GetActor())
 				{
-					GEditor->SelectActor(PinnedActor, /*bInSelected=*/false, /*bNotify=*/false);
+					ActorsToUnpin.Add(PinnedActor);
 				}
 			}
 		}
 
-		WorldPartition->UnpinActors(ActorsToUnpin);
+		constexpr bool bShouldSelect = false;
+		constexpr bool bSelectEvenIfHidden = true;
+		SelectActorsInEditor(ActorsToUnpin, bShouldSelect, bSelectEvenIfHidden);
 
-		GEditor->GetSelectedActors()->EndBatchSelectOperation(/*bNotify=*/true);
+		WorldPartition->UnpinActors(ActorGuidsToUnpin);
 	}
 }
 
@@ -3172,9 +3304,10 @@ void FObjectMixerOutlinerMode::SynchronizeSelection()
 {
 	if (ShouldSyncSelectionFromEditor())
 	{
-		SynchronizeActorSelection();
-		SynchronizeComponentSelection();
+		SynchronizeComponentAndActorSelection();
 		SynchronizeSelectedActorDescs();
+
+		StaticCast<SObjectMixerEditorList*>(SceneOutliner)->GetOnSelectionSynchronized().Broadcast();
 	}
 	
 	bShouldTemporarilyForceSelectionSyncFromEditor = false;
@@ -3283,6 +3416,19 @@ void FObjectMixerOutlinerMode::OnDeleteActorsEnd()
 	SceneOutliner->DeleteFoldersEnd();
 }
 
+bool FObjectMixerOutlinerMode::ShouldAllowHybridRows() const
+{
+	if (const TSharedPtr<FObjectMixerEditorList> PinnedList = GetListModelPtr().Pin())
+	{
+		if (const UObjectMixerObjectFilter* ObjectFilter = PinnedList->GetMainObjectFilterInstance())
+		{
+			return ObjectFilter->ShouldAllowHybridRows();
+		}
+	}
+
+	return true;
+}
+
 FObjectMixerOutlinerModeConfig* FObjectMixerOutlinerMode::GetMutableConfig() const
 {
 	FName OutlinerIdentifier = SceneOutliner->GetOutlinerIdentifier();
diff --git a/Engine/Plugins/Editor/ObjectMixer/ObjectMixer/Source/ObjectMixer/Private/Views/List/ObjectMixerEditorList.cpp b/Engine/Plugins/Editor/ObjectMixer/ObjectMixer/Source/ObjectMixer/Private/Views/List/ObjectMixerEditorList.cpp
--- a/Engine/Plugins/Editor/ObjectMixer/ObjectMixer/Source/ObjectMixer/Private/Views/List/ObjectMixerEditorList.cpp
+++ b/Engine/Plugins/Editor/ObjectMixer/ObjectMixer/Source/ObjectMixer/Private/Views/List/ObjectMixerEditorList.cpp
@@ -4,17 +4,25 @@
 
 #include "ObjectMixerEditorLog.h"
 #include "ObjectMixerEditorModule.h"
+#include "ObjectMixerEditorSettings.h"
 #include "ObjectMixerEditorSerializedData.h"
+#include "SelectionInterface/LevelEditorObjectMixerSelectionInterface.h"
 
 #include "Framework/Commands/GenericCommands.h"
+#include "LevelEditor.h"
 #include "LevelEditorActions.h"
+#include "LevelEditorSequencerIntegration.h"
+#include "Misc/TransactionObjectEvent.h"
 #include "Modules/ModuleManager.h"
+#include "Selection.h"
 #include "Views/List/SObjectMixerEditorList.h"
 
 void FObjectMixerEditorList::Initialize()
 {
 	RegisterAndMapContextMenuCommands();
 
+	BindRefreshDelegates();
+
 	// The base module is always notified when a blueprint filter is compiled.
 	// When it is, all submodules should rebuild their lists.
 	OnBlueprintFilterCompiledHandle = FObjectMixerEditorModule::Get().OnBlueprintFilterCompiled().AddLambda([this]()
@@ -340,9 +348,12 @@ const TSubclassOf<UObjectMixerObjectFilter>& FObjectMixerEditorList::GetDefaultF
 	return DefaultFilterClass;
 }
 
-FObjectMixerEditorList::FObjectMixerEditorList(const FName InModuleName)
+FObjectMixerEditorList::FObjectMixerEditorList(const FName InModuleName, TSharedPtr<IObjectMixerSelectionInterface> InSelectionInterface)
 {
 	ModuleName = InModuleName;
+	SelectionInterface = InSelectionInterface.IsValid()
+		? InSelectionInterface.ToSharedRef()
+		: MakeShared<FLevelEditorObjectMixerSelectionInterface>();
 	
 	OnPostFilterChangeDelegate.AddRaw(this, &FObjectMixerEditorList::OnPostFilterChange);
 }
@@ -351,6 +362,8 @@ FObjectMixerEditorList::~FObjectMixerEditorList()
 {
 	FlushWidget();
 
+	UnbindRefreshDelegates();
+
 	if (FObjectMixerEditorModule* ObjectMixerEditorModule = FModuleManager::GetModulePtr<FObjectMixerEditorModule>("ObjectMixer"))
 	{
 		ObjectMixerEditorModule->OnBlueprintFilterCompiled().Remove(OnBlueprintFilterCompiledHandle);
@@ -593,3 +606,160 @@ void FObjectMixerEditorList::SetCollectionSelected(const FName& CollectionName,
 		}
 	}
 }
+
+void FObjectMixerEditorList::BindRefreshDelegates()
+{
+	check(GEngine && GEditor);
+	DelegateHandles.Add(GEngine->OnLevelActorListChanged().AddLambda([this]()
+	{
+		RequestRebuildList();
+	}));
+
+	DelegateHandles.Add(GEngine->OnActorFolderAdded().AddLambda([this](UActorFolder*)
+	{
+		RequestRebuildList();
+	}));
+
+	DelegateHandles.Add(GEngine->OnActorFoldersUpdatedEvent().AddLambda([this](ULevel*)
+	{
+		RequestRebuildList();
+	}));
+
+	DelegateHandles.Add(GEngine->OnLevelActorFolderChanged().AddLambda([this](const AActor*, FName)
+	{
+		RequestRebuildList();
+	}));
+
+	DelegateHandles.Add(GEngine->OnLevelActorAttached().AddLambda([this](AActor*, const AActor*)
+	{
+		RequestRebuildList();
+	}));
+
+	DelegateHandles.Add(GEngine->OnLevelActorDetached().AddLambda([this](AActor*, const AActor*)
+	{
+		RequestRebuildList();
+	}));
+
+	DelegateHandles.Add(GEditor->OnLevelActorAdded().AddLambda([this](AActor*)
+	{
+		RequestRebuildList();
+	}));
+
+	DelegateHandles.Add(GEditor->OnLevelActorDeleted().AddLambda([this](AActor*)
+	{
+		RequestRebuildList();
+	}));
+
+	DelegateHandles.Add(FEditorDelegates::MapChange.AddLambda([this](uint32)
+	{
+		RequestRebuildList();
+	}));
+
+	DelegateHandles.Add(FEditorDelegates::PostUndoRedo.AddLambda([this]()
+	{
+		RequestRebuildList();
+
+		// Because we can undo/redo collection adds and removes, we need to save the data after each
+		GetMutableDefault<UObjectMixerEditorSerializedData>()->SaveConfig();
+	}));
+
+	DelegateHandles.Add(FCoreUObjectDelegates::OnObjectTransacted.AddLambda([this](UObject*, const FTransactionObjectEvent& Event)
+	{
+		if (Event.GetEventType() == ETransactionObjectEventType::Finalized)
+		{
+			if (Event.HasNameChange())
+			{
+				RequestRebuildList();
+			}
+			else if (const TSet<FName> PropertiesThatRequireRefresh = GetPropertiesThatRequireRefresh(); PropertiesThatRequireRefresh.Num() > 0)
+			{
+				const TSet<FName> ChangedPropertyNames = TSet<FName>(Event.GetChangedProperties());
+
+				if (ChangedPropertyNames.Intersect(PropertiesThatRequireRefresh).Num() > 0)
+				{
+					RequestRebuildList();
+				}
+			}
+		}
+	}));
+
+	FLevelEditorModule& LevelEditorModule = FModuleManager::LoadModuleChecked<FLevelEditorModule>("LevelEditor");
+	DelegateHandles.Add(LevelEditorModule.OnComponentsEdited().AddLambda([this]()
+	{
+		RequestRebuildList();
+	}));
+
+	DelegateHandles.Add(FLevelEditorSequencerIntegration::Get().GetOnSequencersChanged().AddLambda([this]
+	{
+		RequestRebuildList();
+	}));
+
+	DelegateHandles.Add(GetMutableDefault<UObjectMixerEditorSettings>()->OnSettingChanged().AddLambda(
+	[this](UObject*, FPropertyChangedEvent& Event)
+	{
+		if (UObjectMixerEditorSettings::DoesPropertyChangeRequireListRebuild(Event))
+		{
+			RequestRebuildList();
+		}
+	}));
+}
+
+void FObjectMixerEditorList::UnbindRefreshDelegates()
+{
+	// Unbind Delegates
+	if (GEngine)
+	{
+		GEngine->OnLevelActorFolderChanged().RemoveAll(this);
+		GEngine->OnLevelActorListChanged().RemoveAll(this);
+		GEngine->OnActorFolderAdded().RemoveAll(this);
+		GEngine->OnLevelActorDetached().RemoveAll(this);
+		GEngine->OnLevelActorAttached().RemoveAll(this);
+		GEngine->OnActorFolderAdded().RemoveAll(this);
+		GEngine->OnActorFoldersUpdatedEvent().RemoveAll(this);
+	}
+
+	if (GEditor)
+	{
+		GEditor->OnLevelActorAdded().RemoveAll(this);
+		GEditor->OnLevelActorDeleted().RemoveAll(this);
+		GEditor->GetSelectedActors()->SelectionChangedEvent.RemoveAll(this);
+	}
+
+	FEditorDelegates::MapChange.RemoveAll(this);
+	FEditorDelegates::PostUndoRedo.RemoveAll(this);
+	FCoreDelegates::OnPostEngineInit.RemoveAll(this);
+	FCoreUObjectDelegates::OnObjectTransacted.RemoveAll(this);
+
+	if (FLevelEditorModule* LevelEditorModule = FModuleManager::GetModulePtr<FLevelEditorModule>("LevelEditor"))
+	{
+		LevelEditorModule->OnComponentsEdited().RemoveAll(this);
+	}
+
+	FLevelEditorSequencerIntegration::Get().GetOnSequencersChanged().RemoveAll(this);
+
+	if (FModuleManager::Get().IsModuleLoaded(TEXT("ObjectMixer")))
+	{
+		if (UObjectMixerEditorSettings* Settings = GetMutableDefault<UObjectMixerEditorSettings>())
+		{
+			Settings->OnSettingChanged().RemoveAll(this);
+		}
+	}
+
+	for (FDelegateHandle& Delegate : DelegateHandles)
+	{
+		Delegate.Reset();
+	}
+	DelegateHandles.Empty();
+}
+
+TSet<FName> FObjectMixerEditorList::GetPropertiesThatRequireRefresh()
+{
+	TSet<FName> ReturnValue;
+
+	for (const TObjectPtr<UObjectMixerObjectFilter>& Instance : GetObjectFilterInstances())
+	{
+		ReturnValue.Append(Instance->GetPropertiesThatRequireListRefresh());
+	}
+
+	return ReturnValue;
+}
\ No newline at end of file
diff --git a/Engine/Plugins/Editor/ObjectMixer/ObjectMixer/Source/ObjectMixer/Private/Views/List/SObjectMixerEditorList.cpp b/Engine/Plugins/Editor/ObjectMixer/ObjectMixer/Source/ObjectMixer/Private/Views/List/SObjectMixerEditorList.cpp
--- a/Engine/Plugins/Editor/ObjectMixer/ObjectMixer/Source/ObjectMixer/Private/Views/List/SObjectMixerEditorList.cpp
+++ b/Engine/Plugins/Editor/ObjectMixer/ObjectMixer/Source/ObjectMixer/Private/Views/List/SObjectMixerEditorList.cpp
@@ -141,13 +141,10 @@ TSharedRef<SWidget> SObjectMixerEditorList::OnGenerateAddObjectButtonMenu() cons
 	{
 		FMenuBuilder AddObjectButtonMenuBuilder = FMenuBuilder(true, nullptr);
 
-		for (const UClass* Class : ClassesToPlace)
+		for (UClass* Class : ClassesToPlace)
 		{
-			if (const UActorFactory* Factory = GEditor->FindActorFactoryForActorClass(Class))
-			{
-				AddObjectButtonMenuBuilder.AddWidget(
-					SNew(SObjectMixerPlacementAssetMenuEntry, MakeShareable(new FPlaceableItem(*Factory->GetClass()))), FText::GetEmpty());
-			}
+			TSharedPtr<FPlaceableItem> Item = MakeShared<FPlaceableItem>(GEditor->FindActorFactoryForActorClass(Class), FAssetData(Class));
+			AddObjectButtonMenuBuilder.AddWidget(SNew(SObjectMixerPlacementAssetMenuEntry, Item), FText::GetEmpty());
 		}
 
 		return AddObjectButtonMenuBuilder.MakeWidget();
@@ -673,6 +670,17 @@ void SObjectMixerEditorList::Tick(const FGeometry& AllottedGeometry, const doubl
 
 TSharedRef<SWidget> SObjectMixerEditorList::GenerateHeaderRowContextMenu()
 {
+	TSharedPtr<FObjectMixerEditorList> PinnedListModel = ListModelPtr.Pin(); 
+	check(PinnedListModel);
+
+	if (const UObjectMixerObjectFilter* Filter = PinnedListModel->GetMainObjectFilterInstance())
+	{
+		if (!Filter->ShouldAllowColumnCustomizationByUser())
+		{
+			return SNullWidget::NullWidget;
+		}
+	}
+
 	if (!HeaderRowContextMenuWidget)
 	{
 		FMenuBuilder MenuBuilder(false, nullptr);
diff --git a/Engine/Plugins/Editor/ObjectMixer/ObjectMixer/Source/ObjectMixer/Public/ObjectFilter/ObjectMixerEditorObjectFilter.h b/Engine/Plugins/Editor/ObjectMixer/ObjectMixer/Source/ObjectMixer/Public/ObjectFilter/ObjectMixerEditorObjectFilter.h
--- a/Engine/Plugins/Editor/ObjectMixer/ObjectMixer/Source/ObjectMixer/Public/ObjectFilter/ObjectMixerEditorObjectFilter.h
+++ b/Engine/Plugins/Editor/ObjectMixer/ObjectMixer/Source/ObjectMixer/Public/ObjectFilter/ObjectMixerEditorObjectFilter.h
@@ -76,6 +76,16 @@ public:
 	 */
 	virtual bool GetShowTransientObjects() const;
 
+	/**
+	 * Determines if hybrid rows should be used when enabled in user settings. True by default.
+	 */
+	virtual bool ShouldAllowHybridRows() const;
+
+	/**
+	 * Determines if the user should be able to change which columns are displayed. True by default.
+	 */
+	virtual bool ShouldAllowColumnCustomizationByUser() const;
+
 	/**
 	 * Specify a list of property names corresponding to columns you want to show by default.
 	 * For example, you can specify "Intensity" and "LightColor" to show only those property columns by default in the UI.
diff --git a/Engine/Plugins/Editor/ObjectMixer/ObjectMixer/Source/ObjectMixer/Public/ObjectMixerEditorModule.h b/Engine/Plugins/Editor/ObjectMixer/ObjectMixer/Source/ObjectMixer/Public/ObjectMixerEditorModule.h
--- a/Engine/Plugins/Editor/ObjectMixer/ObjectMixer/Source/ObjectMixer/Public/ObjectMixerEditorModule.h
+++ b/Engine/Plugins/Editor/ObjectMixer/ObjectMixer/Source/ObjectMixer/Public/ObjectMixerEditorModule.h
@@ -47,17 +47,6 @@ public:
 	 */
 	bool RegenerateListWidget();
 
-	/**
-	 * Regenerate the list items and refresh the list. Call when adding or removing variables.
-	 */
-	virtual void RequestRebuildList() const;
-	
-	/**
-	 * Refresh filters and sorting.
-	 * Useful for when the list state has gone stale but the variable count has not changed.
-	 */
-	virtual void RefreshList() const;
-
 	/** Called when the Rename command is executed from the UI or hotkey. */
 	virtual void OnRenameCommand();
 	
@@ -97,11 +86,6 @@ public:
 
 protected:
 
-	virtual void BindDelegates();
-
-	/** If a property is changed that has a name found in this set, the panel will be refreshed. */
-	TSet<FName> GetPropertiesThatRequireRefresh() const;
-	
 	/** Lives for as long as the module is loaded. */
 	TSharedPtr<FObjectMixerEditorList> ListModel;
 
@@ -122,8 +106,6 @@ protected:
 	 * This filter class cannot be turned off by the end user.
 	 */
 	TSubclassOf<UObjectMixerObjectFilter> DefaultFilterClass;
-	
-	TSet<FDelegateHandle> DelegateHandles;
 
 	FOnBlueprintFilterCompiled OnBlueprintFilterCompiledDelegate;
 
diff --git a/Engine/Plugins/Editor/ObjectMixer/ObjectMixer/Source/ObjectMixer/Public/SelectionInterface/IObjectMixerSelectionInterface.h b/Engine/Plugins/Editor/ObjectMixer/ObjectMixer/Source/ObjectMixer/Public/SelectionInterface/IObjectMixerSelectionInterface.h
new file mode 100644
--- /dev/null
+++ b/Engine/Plugins/Editor/ObjectMixer/ObjectMixer/Source/ObjectMixer/Public/SelectionInterface/IObjectMixerSelectionInterface.h
@@ -0,0 +1,47 @@
+// Copyright Epic Games, Inc. All Rights Reserved.
+
+#pragma once
+
+#include "Delegates/Delegate.h"
+#include "HAL/Platform.h"
+#include "Templates/SharedPointer.h"
+
+class AActor;
+class UActorComponent;
+
+/**
+ * Provides an interface for synchronizing an ObjectMixer's selection with another part of the editor
+ */
+class OBJECTMIXEREDITOR_API IObjectMixerSelectionInterface : public TSharedFromThis<IObjectMixerSelectionInterface>
+{
+public:
+	DECLARE_MULTICAST_DELEGATE(FOnSelectionChanged);
+
+public:
+	virtual ~IObjectMixerSelectionInterface() = default;
+
+	/**
+	 * Select or deselect a list of components.
+	 * @param	InSelectedActors		The actors to select.
+	 * @param	bShouldSelect			If true, select the actors and deselect all other actors. If false, only deselect the listed actors.
+	 * @param	bSelectEvenIfHidden		If false, only select actors that are visible (e.g. not filtered out) in the editor.
+	 */
+	virtual void SelectActors(const TArray<AActor*>& InSelectedActors, bool bShouldSelect, bool bSelectEvenIfHidden) = 0;
+
+	/**
+	 * Select or deselect a list of components.
+	 * @param	InSelectedComponents	The components to select.
+	 * @param	bShouldSelect			If true, select the components and deselect all other components. If false, only deselect the listed components.
+	 * @param	bSelectEvenIfHidden		If false, only select components that are visible (e.g. not filtered out) in the editor.
+	 */
+	virtual void SelectComponents(const TArray<UActorComponent*>& InSelectedComponents, bool bShouldSelect, bool bSelectEvenIfHidden) = 0;
+
+	/** Get the list of selected actors. */
+	virtual TArray<AActor*> GetSelectedActors() const = 0;
+
+	/** Get the list of selected components. */
+	virtual TArray<UActorComponent*> GetSelectedComponents() const = 0;
+
+	/** Get the event handler that is raised when the synchronized selection changes. */
+	virtual FOnSelectionChanged& OnSelectionChanged() = 0;
+};
\ No newline at end of file
diff --git a/Engine/Plugins/Editor/ObjectMixer/ObjectMixer/Source/ObjectMixer/Public/SelectionInterface/LevelEditorObjectMixerSelectionInterface.h b/Engine/Plugins/Editor/ObjectMixer/ObjectMixer/Source/ObjectMixer/Public/SelectionInterface/LevelEditorObjectMixerSelectionInterface.h
new file mode 100644
--- /dev/null
+++ b/Engine/Plugins/Editor/ObjectMixer/ObjectMixer/Source/ObjectMixer/Public/SelectionInterface/LevelEditorObjectMixerSelectionInterface.h
@@ -0,0 +1,29 @@
+// Copyright Epic Games, Inc. All Rights Reserved.
+
+#pragma once
+
+#include "IObjectMixerSelectionInterface.h"
+
+/**
+ * Provides an interface for the ObjectMixer to synchronize with the level editor's selections (via GEditor)
+ */
+class OBJECTMIXEREDITOR_API FLevelEditorObjectMixerSelectionInterface : public IObjectMixerSelectionInterface
+{
+public:
+
+	FLevelEditorObjectMixerSelectionInterface();
+	virtual ~FLevelEditorObjectMixerSelectionInterface() override;
+
+	//~ Begin IObjectMixerSelectionInterface interface
+	virtual void SelectActors(const TArray<AActor*>& InSelectedActors, bool bShouldSelect, bool bSelectEvenIfHidden) override;
+	virtual void SelectComponents(const TArray<UActorComponent*>& InSelectedComponents, bool bShouldSelect, bool bSelectEvenIfHidden) override;
+	virtual TArray<AActor*> GetSelectedActors() const override;
+	virtual TArray<UActorComponent*> GetSelectedComponents() const override;
+	virtual FOnSelectionChanged& OnSelectionChanged() override { return SelectionChanged; }
+	//~ End IObjectMixerSelectionInterface interface
+
+private:
+	void OnLevelSelectionChanged(UObject* Obj);
+
+	FOnSelectionChanged SelectionChanged;
+};
\ No newline at end of file
diff --git a/Engine/Plugins/Editor/ObjectMixer/ObjectMixer/Source/ObjectMixer/Public/Views/List/Modes/ObjectMixerOutlinerMode.h b/Engine/Plugins/Editor/ObjectMixer/ObjectMixer/Source/ObjectMixer/Public/Views/List/Modes/ObjectMixerOutlinerMode.h
--- a/Engine/Plugins/Editor/ObjectMixer/ObjectMixer/Source/ObjectMixer/Public/Views/List/Modes/ObjectMixerOutlinerMode.h
+++ b/Engine/Plugins/Editor/ObjectMixer/ObjectMixer/Source/ObjectMixer/Public/Views/List/Modes/ObjectMixerOutlinerMode.h
@@ -13,28 +13,40 @@
 #include "ObjectMixerOutlinerMode.generated.h"
 
 class FObjectMixerEditorList;
+class FObjectMixerOutlinerMode;
 class IWorldPartitionEditorModule;
 
 namespace ObjectMixerOutliner
 {
+	/** A row selector that depends on whether hybrid rows are allowed */
+	struct FHybridRowSelector
+	{
+		FHybridRowSelector(const FObjectMixerOutlinerMode* Mode);
+
+		bool ShouldAllowHybridRows() const { return bAllowHybridRows; }
+
+	private:
+		bool bAllowHybridRows = true;
+	};
+
 	struct FWeakActorSelectorAcceptingComponents
 	{
 		bool operator()(const TWeakPtr<ISceneOutlinerTreeItem>& Item, TWeakObjectPtr<AActor>& DataOut) const;
 	};
 
-	struct FComponentSelector
+	struct FComponentSelector : FHybridRowSelector
 	{
 		bool operator()(const TWeakPtr<ISceneOutlinerTreeItem>& Item, UActorComponent*& DataOut) const;
 	};
 
 	/** Functor which can be used to get weak actor pointers from a selection */
-	struct FWeakActorSelector
+	struct FWeakActorSelector 
 	{
 		bool operator()(const TWeakPtr<ISceneOutlinerTreeItem>& Item, TWeakObjectPtr<AActor>& DataOut) const;
 	};
 
 	/** Functor which can be used to get actors from a selection including component parents */
-	struct FActorSelector
+	struct FActorSelector : FHybridRowSelector
 	{
 		bool operator()(const TWeakPtr<ISceneOutlinerTreeItem>& Item, AActor*& ActorPtrOut) const;
 	};
@@ -269,6 +281,9 @@ public:
 	
 	/** Function called by the Outliner Filter Bar to compare an item with Type Filters*/
 	virtual bool CompareItemWithClassName(SceneOutliner::FilterBarType InItem, const TSet<FTopLevelAssetPath>&) const override;
+
+	/** Check whether hybrid rows are allowed for the associated object list */
+	bool ShouldAllowHybridRows() const;
 	
 protected:
 
@@ -277,7 +292,7 @@ protected:
 	void OnMapChange(uint32 MapFlags);
 	void OnNewCurrentLevel();
 
-	void OnLevelSelectionChanged(UObject* Obj);
+	void OnEditorSelectionChanged();
 	void OnActorLabelChanged(AActor* ChangedActor);
 	void OnObjectsReplaced(const TMap<UObject*, UObject*>& ReplacementMap);
 	void OnLevelActorRequestsRename(const AActor* Actor);
@@ -288,8 +303,16 @@ protected:
 	void SynchronizeAllSelectionsToEditor();
 	bool HasActorSelectionChanged(TArray<AActor*>& OutSelectedActors, bool& bOutAreAnyInPIE);
 	bool HasComponentSelectionChanged(TArray<UActorComponent*>& OutSelectedComponents, bool& bOutAreAnyInPIE);
-	static void SelectActorsInEditor(const TArray<AActor*>& InSelectedActors);
-	static void SelectComponentsInEditor(const TArray<UActorComponent*>& InSelectedComponents);
+	void SelectActorsInEditor(const TArray<AActor*>& InSelectedActors, bool bShouldSelect, bool bSelectEvenIfHidden);
+	void SelectComponentsInEditor(const TArray<UActorComponent*>& InSelectedComponents, bool bShouldSelect, bool bSelectEvenIfHidden);
+	void SelectActorsInMixer(const TArray<AActor*>& InSelectedActors, bool bShouldSelect, bool bSelectEvenIfHidden);
+	void SelectComponentsInMixer(const TArray<UActorComponent*>& InSelectedComponents, bool bShouldSelect, bool bSelectEvenIfHidden);
+
+	/** Get the list of actors selected in the editor */
+	TArray<AActor*> GetSelectedActorsInEditor() const;
+
+	/** Get the list of components selected in the editor */
+	TArray<UActorComponent*> GetSelectedComponentsInEditor() const;
 	
 	/** Build and up the context menu */
 	TSharedPtr<SWidget> BuildContextMenu();
@@ -300,9 +323,12 @@ protected:
 	bool GetFolderNamesFromPayload(const FSceneOutlinerDragDropPayload& InPayload, TArray<FName>& OutFolders, FFolder::FRootObject& OutCommonRootObject) const;
 	FFolder GetWorldDefaultRootFolder() const;
 
-	void SynchronizeComponentSelection();
+	/** Synchronize both the actor and component selection with the editor, with components taking priority over actors */
+	void SynchronizeComponentAndActorSelection();
 	void SynchronizeSelectedActorDescs();
 
+	void SynchronizeCustomActorSelection();
+
 	void OnActorEditorContextSubsystemChanged();
 
 	/** Filter factories */
diff --git a/Engine/Plugins/Editor/ObjectMixer/ObjectMixer/Source/ObjectMixer/Public/Views/List/ObjectMixerEditorList.h b/Engine/Plugins/Editor/ObjectMixer/ObjectMixer/Source/ObjectMixer/Public/Views/List/ObjectMixerEditorList.h
--- a/Engine/Plugins/Editor/ObjectMixer/ObjectMixer/Source/ObjectMixer/Public/Views/List/ObjectMixerEditorList.h
+++ b/Engine/Plugins/Editor/ObjectMixer/ObjectMixer/Source/ObjectMixer/Public/Views/List/ObjectMixerEditorList.h
@@ -3,6 +3,7 @@
 #pragma once
 
 #include "ObjectMixerEditorListRowData.h"
+#include "SelectionInterface/IObjectMixerSelectionInterface.h"
 
 #include "Templates/SharedPointer.h"
 #include "Widgets/SWidget.h"
@@ -19,7 +20,7 @@ class OBJECTMIXEREDITOR_API FObjectMixerEditorList : public TSharedFromThis<FObj
 {
 public:
 
-	FObjectMixerEditorList(const FName InModuleName);
+	FObjectMixerEditorList(const FName InModuleName, TSharedPtr<IObjectMixerSelectionInterface> InSelectionInterface = nullptr);
 
 	virtual ~FObjectMixerEditorList();
 
@@ -45,6 +46,9 @@ public:
 	 */
 	void RefreshList() const;
 
+	/** Get the selection interface used to synchronize this with the rest of the editor */
+	TSharedPtr<IObjectMixerSelectionInterface> GetSelectionInterface() const { return SelectionInterface; }
+
 	void BuildPerformanceCache();
 
 	bool ShouldShowTransientObjects() const;
@@ -189,6 +193,15 @@ public:
 	bool bShouldShowTransientObjectsCache = false;
 
 protected:
+
+	/** Bind delegates to refresh the list when editor state changes. */
+	virtual void BindRefreshDelegates();
+
+	/** Unbind delegates to refresh the list when editor state changes. */
+	virtual void UnbindRefreshDelegates();
+
+	/** If a property is changed that has a name found in this set, the list will be refreshed. */
+	TSet<FName> GetPropertiesThatRequireRefresh();
 	
 	virtual void AddReferencedObjects( FReferenceCollector& Collector )  override
 	{
@@ -226,4 +239,10 @@ protected:
 	FName ModuleName = NAME_None;
 	
 	FDelegateHandle OnBlueprintFilterCompiledHandle;
+
+	/** Interface used to synchronize selection with another part of the editor */
+	TSharedPtr<IObjectMixerSelectionInterface> SelectionInterface;
+
+	/** Handles for delegates registered to refresh the list */
+	TSet<FDelegateHandle> DelegateHandles;
 };
diff --git a/Engine/Plugins/Editor/ObjectMixer/ObjectMixer/Source/ObjectMixer/Public/Views/List/SObjectMixerEditorList.h b/Engine/Plugins/Editor/ObjectMixer/ObjectMixer/Source/ObjectMixer/Public/Views/List/SObjectMixerEditorList.h
--- a/Engine/Plugins/Editor/ObjectMixer/ObjectMixer/Source/ObjectMixer/Public/Views/List/SObjectMixerEditorList.h
+++ b/Engine/Plugins/Editor/ObjectMixer/ObjectMixer/Source/ObjectMixer/Public/Views/List/SObjectMixerEditorList.h
@@ -71,6 +71,7 @@ class OBJECTMIXEREDITOR_API SObjectMixerEditorList : public SSceneOutliner
 {
 
 public:
+	DECLARE_EVENT(SObjectMixerEditorList, FSelectionSynchronizedEvent);
 	
 	// Columns
 	static const FName ItemNameColumnName;
@@ -162,6 +163,9 @@ public:
 	void RestoreDefaultPropertyColumns();
 	TSharedRef<SWidget> GenerateHeaderRowContextMenu();
 
+	/** Get the event that broadcasts whenever the selection list is synchronized from the editor */
+	FSelectionSynchronizedEvent& GetOnSelectionSynchronized() { return OnSelectionSynchronizedEvent; }
+
 	/* Begin SSceneOutliner Interface */
 	virtual void CustomAddToToolbar(TSharedPtr<class SHorizontalBox> Toolbar) override;
 	/* End SSceneOutliner Interface */
@@ -184,6 +188,8 @@ protected:
 
 	TWeakPtr<SHorizontalBox> ToolbarPtr;
 
+	FSelectionSynchronizedEvent OnSelectionSynchronizedEvent;
+
 	bool CanCreateFolder() const;
 	TSharedRef<SWidget> OnGenerateAddObjectButtonMenu() const;
 	
diff --git a/Engine/Plugins/Runtime/nDisplay/Source/DisplayClusterColorGrading/DisplayClusterColorGrading.Build.cs b/Engine/Plugins/Runtime/nDisplay/Source/DisplayClusterColorGrading/DisplayClusterColorGrading.Build.cs
--- a/Engine/Plugins/Runtime/nDisplay/Source/DisplayClusterColorGrading/DisplayClusterColorGrading.Build.cs
+++ b/Engine/Plugins/Runtime/nDisplay/Source/DisplayClusterColorGrading/DisplayClusterColorGrading.Build.cs
@@ -24,6 +24,7 @@ public class DisplayClusterColorGrading : ModuleRules
 				"Engine",
 				"InputCore",
 				"Kismet",
+				"ObjectMixerEditor",
 				"PropertyEditor",
 				"PropertyPath",
 				"Slate",
diff --git a/Engine/Plugins/Runtime/nDisplay/Source/DisplayClusterColorGrading/Private/DisplayClusterColorGradingModule.cpp b/Engine/Plugins/Runtime/nDisplay/Source/DisplayClusterColorGrading/Private/DisplayClusterColorGradingModule.cpp
--- a/Engine/Plugins/Runtime/nDisplay/Source/DisplayClusterColorGrading/Private/DisplayClusterColorGradingModule.cpp
+++ b/Engine/Plugins/Runtime/nDisplay/Source/DisplayClusterColorGrading/Private/DisplayClusterColorGradingModule.cpp
@@ -11,6 +11,7 @@
 #include "Components/DisplayClusterICVFXCameraComponent.h"
 
 #include "ColorCorrectRegion.h"
+#include "ColorGradingMixerObjectFilterRegistry.h"
 #include "Engine/PostProcessVolume.h"
 
 #define LOCTEXT_NAMESPACE "DisplayClusterColorGrading"
@@ -28,6 +29,11 @@ void FDisplayClusterColorGradingModule::StartupModule()
 	FColorGradingListItem::RegisterColorGradingListItemGenerator<FDisplayClusterColorGradingListItemGenerator_RootActor>();
 	FColorGradingListItem::RegisterColorGradingListItemGenerator<FDisplayClusterColorGradingListItemGenerator_ICVFXCamera>();
 
+	FColorGradingMixerObjectFilterRegistry::RegisterActorClassToPlace(ADisplayClusterRootActor::StaticClass());
+
+	FColorGradingMixerObjectFilterRegistry::RegisterObjectClassToFilter(ADisplayClusterRootActor::StaticClass());
+	FColorGradingMixerObjectFilterRegistry::RegisterObjectClassToFilter(UDisplayClusterICVFXCameraComponent::StaticClass());
+
 	FDisplayClusterColorGradingCommands::Register();
 }
 
diff --git a/Engine/Plugins/Runtime/nDisplay/Source/DisplayClusterColorGrading/Private/Drawer/SDisplayClusterColorGradingDrawer.cpp b/Engine/Plugins/Runtime/nDisplay/Source/DisplayClusterColorGrading/Private/Drawer/SDisplayClusterColorGradingDrawer.cpp
--- a/Engine/Plugins/Runtime/nDisplay/Source/DisplayClusterColorGrading/Private/Drawer/SDisplayClusterColorGradingDrawer.cpp
+++ b/Engine/Plugins/Runtime/nDisplay/Source/DisplayClusterColorGrading/Private/Drawer/SDisplayClusterColorGradingDrawer.cpp
@@ -14,6 +14,7 @@
 #include "DisplayClusterRootActor.h"
 #include "Components/DisplayClusterICVFXCameraComponent.h"
 #include "DisplayClusterConfigurationTypes.h"
+#include "SelectionInterface/DisplayClusterObjectMixerSelectionInterface.h"
 
 #include "ColorCorrectRegion.h"
 #include "ColorCorrectWindow.h"
@@ -57,12 +58,13 @@ void SDisplayClusterColorGradingDrawer::Construct(const FArguments& InArgs, bool
 	ChildSlot
 	[
 		SAssignNew(MainPanel, SColorGradingPanel)
+		.SelectionInterface(MakeShared<FDisplayClusterObjectMixerSelectionInterface>())
 		.OverrideWorld(this, &SDisplayClusterColorGradingDrawer::GetOperatorWorld)
 		.IsInDrawer(bInIsInDrawer)
 		.OnDocked_Lambda([]() {
 			IDisplayClusterColorGrading::Get().GetColorGradingDrawerSingleton().DockColorGradingDrawer();
 		})
-		.ActorFilter([SharedThis, this](AActor* Actor) -> bool {
+		.ActorFilter([SharedThis, this](const AActor* Actor) -> bool {
 			if (SharedThis.IsValid() && OperatorViewModel.IsValid() && Actor && Actor->IsA<ADisplayClusterRootActor>())
 			{
 				return Actor == OperatorViewModel->GetRootActor();
diff --git a/Engine/Plugins/Runtime/nDisplay/Source/DisplayClusterColorGrading/Private/SelectionInterface/DisplayClusterObjectMixerSelectionInterface.cpp b/Engine/Plugins/Runtime/nDisplay/Source/DisplayClusterColorGrading/Private/SelectionInterface/DisplayClusterObjectMixerSelectionInterface.cpp
new file mode 100644
--- /dev/null
+++ b/Engine/Plugins/Runtime/nDisplay/Source/DisplayClusterColorGrading/Private/SelectionInterface/DisplayClusterObjectMixerSelectionInterface.cpp
@@ -0,0 +1,81 @@
+// Copyright Epic Games, Inc. All Rights Reserved.
+
+#include "SelectionInterface/DisplayClusterObjectMixerSelectionInterface.h"
+
+#include "Algo/Transform.h"
+#include "Editor/UnrealEdEngine.h"
+#include "Engine/Selection.h"
+#include "IDisplayClusterOperator.h"
+#include "IDisplayClusterOperatorViewModel.h"
+
+FDisplayClusterObjectMixerSelectionInterface::FDisplayClusterObjectMixerSelectionInterface()
+{
+	IDisplayClusterOperator::Get().GetOperatorViewModel()->OnOutlinerSelectionChanged().AddRaw(
+		this, &FDisplayClusterObjectMixerSelectionInterface::OnOutlinerSelectionChanged
+	);
+}
+
+FDisplayClusterObjectMixerSelectionInterface::~FDisplayClusterObjectMixerSelectionInterface()
+{
+	IDisplayClusterOperator::Get().GetOperatorViewModel()->OnOutlinerSelectionChanged().RemoveAll(this);
+}
+
+void FDisplayClusterObjectMixerSelectionInterface::SelectActors(const TArray<AActor*>& InSelectedActors, bool bShouldSelect, bool bSelectEvenIfHidden)
+{
+	TSharedRef<IDisplayClusterOperatorViewModel> Operator = IDisplayClusterOperator::Get().GetOperatorViewModel();
+
+	Operator->SelectActors(InSelectedActors, bShouldSelect);
+}
+
+void FDisplayClusterObjectMixerSelectionInterface::SelectComponents(const TArray<UActorComponent*>& InSelectedComponents, bool bShouldSelect, bool bSelectEvenIfHidden)
+{
+	TSharedRef<IDisplayClusterOperatorViewModel> Operator = IDisplayClusterOperator::Get().GetOperatorViewModel();
+	TArray<UObject*> DetailObjects;
+
+	if (bShouldSelect)
+	{
+		DetailObjects.Append(InSelectedComponents);
+	}
+	else
+	{
+		Algo::Transform(Operator->GetDetailObjects(), DetailObjects, [](const TWeakObjectPtr<UObject>& Object) { return Object.Get(); });
+
+		for (UActorComponent* Component : InSelectedComponents)
+		{
+			DetailObjects.RemoveSingle(Component);
+		}
+	}
+
+	Operator->ShowDetailsForObjects(DetailObjects);
+}
+
+TArray<AActor*> FDisplayClusterObjectMixerSelectionInterface::GetSelectedActors() const
+{
+	return SelectedActors;
+}
+
+TArray<UActorComponent*> FDisplayClusterObjectMixerSelectionInterface::GetSelectedComponents() const
+{
+	TArray<UActorComponent*> OutComponents;
+
+	Algo::TransformIf(
+		IDisplayClusterOperator::Get().GetOperatorViewModel()->GetDetailObjects(),
+		OutComponents,
+		[](const TWeakObjectPtr<UObject>& Object) -> bool
+		{
+			return Object.IsValid() && Object->IsA<UActorComponent>();
+		},
+		[](const TWeakObjectPtr<UObject>& Object) -> UActorComponent*
+		{
+			return Cast<UActorComponent>(Object.Get());
+		}
+	);
+
+	return MoveTemp(OutComponents);
+}
+
+void FDisplayClusterObjectMixerSelectionInterface::OnOutlinerSelectionChanged(const TArray<AActor*>& Actors)
+{
+	SelectedActors = Actors;
+	SelectionChanged.Broadcast();
+}
diff --git a/Engine/Plugins/Runtime/nDisplay/Source/DisplayClusterColorGrading/Private/SelectionInterface/DisplayClusterObjectMixerSelectionInterface.h b/Engine/Plugins/Runtime/nDisplay/Source/DisplayClusterColorGrading/Private/SelectionInterface/DisplayClusterObjectMixerSelectionInterface.h
new file mode 100644
--- /dev/null
+++ b/Engine/Plugins/Runtime/nDisplay/Source/DisplayClusterColorGrading/Private/SelectionInterface/DisplayClusterObjectMixerSelectionInterface.h
@@ -0,0 +1,32 @@
+// Copyright Epic Games, Inc. All Rights Reserved.
+
+#pragma once
+
+#include "SelectionInterface/IObjectMixerSelectionInterface.h"
+
+/**
+ * Provides an interface for the ObjectMixer to access DisplayCluster operator panel selections
+ */
+class FDisplayClusterObjectMixerSelectionInterface : public IObjectMixerSelectionInterface
+{
+public:
+
+	FDisplayClusterObjectMixerSelectionInterface();
+	virtual ~FDisplayClusterObjectMixerSelectionInterface() override;
+
+	//~ Begin IObjectMixerSelectionInterface interface
+	virtual void SelectActors(const TArray<AActor*>& InSelectedActors, bool bShouldSelect, bool bSelectEvenIfHidden) override;
+	virtual void SelectComponents(const TArray<UActorComponent*>& InSelectedComponents, bool bShouldSelect, bool bSelectEvenIfHidden) override;
+	virtual TArray<AActor*> GetSelectedActors() const override;
+	virtual TArray<UActorComponent*> GetSelectedComponents() const override;
+	virtual FOnSelectionChanged& OnSelectionChanged() override { return SelectionChanged; }
+	//~ End IObjectMixerSelectionInterface interface
+
+private:
+	void OnOutlinerSelectionChanged(const TArray<AActor*>& Actors);
+
+	FOnSelectionChanged SelectionChanged;
+
+	/** Last reported list of actors selected in the operator panel's outliner */
+	TArray<AActor*> SelectedActors;
+};
\ No newline at end of file
diff --git a/Engine/Plugins/Runtime/nDisplay/Source/DisplayClusterLightCardEditor/Private/DisplayClusterLightCardEditor.cpp b/Engine/Plugins/Runtime/nDisplay/Source/DisplayClusterLightCardEditor/Private/DisplayClusterLightCardEditor.cpp
--- a/Engine/Plugins/Runtime/nDisplay/Source/DisplayClusterLightCardEditor/Private/DisplayClusterLightCardEditor.cpp
+++ b/Engine/Plugins/Runtime/nDisplay/Source/DisplayClusterLightCardEditor/Private/DisplayClusterLightCardEditor.cpp
@@ -76,6 +76,7 @@ FDisplayClusterLightCardEditor::~FDisplayClusterLightCardEditor()
 	UnregisterTabSpawners();
 	
 	IDisplayClusterOperator::Get().GetOperatorViewModel()->OnActiveRootActorChanged().Remove(ActiveRootActorChangedHandle);
+	IDisplayClusterOperator::Get().GetOperatorViewModel()->OnActorsSelectedExternally().Remove(ActorsSelectedExternallyHandle);
 	FCoreUObjectDelegates::OnObjectPropertyChanged.RemoveAll(this);
 
 	if (GEngine != nullptr)
@@ -109,6 +110,7 @@ void FDisplayClusterLightCardEditor::Initialize(TSharedRef<IDisplayClusterOperat
 	OperatorViewModel = InViewModel;
 
 	ActiveRootActorChangedHandle = InViewModel->OnActiveRootActorChanged().AddSP(this, &FDisplayClusterLightCardEditor::OnActiveRootActorChanged);
+	ActorsSelectedExternallyHandle = InViewModel->OnActorsSelectedExternally().AddSP(this, &FDisplayClusterLightCardEditor::OnActorsSelectedExternally);
 	if (GEngine != nullptr)
 	{
 		GEngine->OnLevelActorAdded().AddSP(this, &FDisplayClusterLightCardEditor::OnLevelActorAdded);
@@ -1054,6 +1056,31 @@ void FDisplayClusterLightCardEditor::OnActiveRootActorChanged(ADisplayClusterRoo
 	RefreshLabels();
 }
 
+void FDisplayClusterLightCardEditor::OnActorsSelectedExternally(const TArray<AActor*>& Actors, bool bShouldSelect)
+{
+	if (!OperatorViewModel.IsValid())
+	{
+		return;
+	}
+
+	if (bShouldSelect)
+	{
+		LightCardOutliner->SelectActors(Actors);
+	}
+	else
+	{
+		TArray<AActor*> NewActors;
+		LightCardOutliner->GetSelectedActors(NewActors);
+		
+		for (AActor* Actor : Actors)
+		{
+			NewActors.Remove(Actor);
+		}
+
+		LightCardOutliner->SelectActors(NewActors);
+	}
+}
+
 void FDisplayClusterLightCardEditor::RegisterTabSpawners()
 {
 	IDisplayClusterOperator::Get().OnRegisterLayoutExtensions().AddStatic(&FDisplayClusterLightCardEditor::RegisterLayoutExtension);
diff --git a/Engine/Plugins/Runtime/nDisplay/Source/DisplayClusterLightCardEditor/Private/DisplayClusterLightCardEditor.h b/Engine/Plugins/Runtime/nDisplay/Source/DisplayClusterLightCardEditor/Private/DisplayClusterLightCardEditor.h
--- a/Engine/Plugins/Runtime/nDisplay/Source/DisplayClusterLightCardEditor/Private/DisplayClusterLightCardEditor.h
+++ b/Engine/Plugins/Runtime/nDisplay/Source/DisplayClusterLightCardEditor/Private/DisplayClusterLightCardEditor.h
@@ -221,6 +221,9 @@ public:
 private:
 	/** Raised when the active Display cluster root actor has been changed in the operator panel */
 	void OnActiveRootActorChanged(ADisplayClusterRootActor* NewRootActor);
+
+	/** Raised when a request to select or deselect actors in the operator panel is raised from outside the panel */
+	void OnActorsSelectedExternally(const TArray<AActor*>& Actors, bool bShouldSelect);
 	
 	/** Registers the light card editor with the global tab manager and adds it to the operator panel's extension tab stack */
 	void RegisterTabSpawners();
@@ -362,6 +365,9 @@ private:
 	/** Delegate handle for the OnActiveRootActorChanged delegate */
 	FDelegateHandle ActiveRootActorChangedHandle;
 
+	/** Delegate handle for the OnActorsSelectedExternally delegate */
+	FDelegateHandle ActorsSelectedExternallyHandle;
+
 	/** Delegate handle for when an object is transacted */
 	FDelegateHandle OnObjectTransactedHandle;
 };
\ No newline at end of file
diff --git a/Engine/Plugins/Runtime/nDisplay/Source/DisplayClusterLightCardEditor/Private/Outliner/SDisplayClusterLightCardOutliner.cpp b/Engine/Plugins/Runtime/nDisplay/Source/DisplayClusterLightCardEditor/Private/Outliner/SDisplayClusterLightCardOutliner.cpp
--- a/Engine/Plugins/Runtime/nDisplay/Source/DisplayClusterLightCardEditor/Private/Outliner/SDisplayClusterLightCardOutliner.cpp
+++ b/Engine/Plugins/Runtime/nDisplay/Source/DisplayClusterLightCardEditor/Private/Outliner/SDisplayClusterLightCardOutliner.cpp
@@ -8,6 +8,7 @@
 #include "DisplayClusterLightCardOutlinerColumns.h"
 
 #include "IDisplayClusterOperator.h"
+#include "IDisplayClusterOperatorViewModel.h"
 
 #include "DisplayClusterConfigurationTypes.h"
 
@@ -90,6 +91,8 @@ void SDisplayClusterLightCardOutliner::GetSelectedActors(TArray<AActor*>& OutSel
 
 void SDisplayClusterLightCardOutliner::SelectActors(const TArray<AActor*>& ActorsToSelect)
 {
+	IDisplayClusterOperator::Get().GetOperatorViewModel()->OnOutlinerSelectionChanged().Broadcast(ActorsToSelect);
+
 	if (bIsOutlinerChangingSelection)
 	{
 		// Selection change initiated from outliner
diff --git a/Engine/Plugins/Runtime/nDisplay/Source/DisplayClusterOperator/Private/DisplayClusterOperatorViewModel.cpp b/Engine/Plugins/Runtime/nDisplay/Source/DisplayClusterOperator/Private/DisplayClusterOperatorViewModel.cpp
--- a/Engine/Plugins/Runtime/nDisplay/Source/DisplayClusterOperator/Private/DisplayClusterOperatorViewModel.cpp
+++ b/Engine/Plugins/Runtime/nDisplay/Source/DisplayClusterOperator/Private/DisplayClusterOperatorViewModel.cpp
@@ -57,6 +57,11 @@ void FDisplayClusterOperatorViewModel::ShowDetailsForObjects(const TArray<UObjec
 	}
 }
 
+void FDisplayClusterOperatorViewModel::SelectActors(const TArray<AActor*>& Actors, bool bShouldSelect)
+{
+	ActorsSelectedExternally.Broadcast(Actors, bShouldSelect);
+}
+
 TSharedRef<FTabManager> FDisplayClusterOperatorViewModel::CreateTabManager(const TSharedRef<SDockTab>& MajorTabOwner)
 {
 	TabManager = FGlobalTabmanager::Get()->NewTabManager(MajorTabOwner);
diff --git a/Engine/Plugins/Runtime/nDisplay/Source/DisplayClusterOperator/Private/DisplayClusterOperatorViewModel.h b/Engine/Plugins/Runtime/nDisplay/Source/DisplayClusterOperator/Private/DisplayClusterOperatorViewModel.h
--- a/Engine/Plugins/Runtime/nDisplay/Source/DisplayClusterOperator/Private/DisplayClusterOperatorViewModel.h
+++ b/Engine/Plugins/Runtime/nDisplay/Source/DisplayClusterOperator/Private/DisplayClusterOperatorViewModel.h
@@ -23,7 +23,10 @@ public:
 	virtual const TArray<TWeakObjectPtr<UObject>> GetDetailObjects() override { return DetailObjects; }
 	virtual void ShowDetailsForObject(UObject* Object) override;
 	virtual void ShowDetailsForObjects(const TArray<UObject*>& Objects) override;
+	virtual void SelectActors(const TArray<AActor*>& Actors, bool bShouldSelect) override;
 	virtual FOnDetailObjectsChanged& OnDetailObjectsChanged() override { return DetailObjectsChanged; }
+	virtual FOnOutlinerSelectionChanged& OnOutlinerSelectionChanged() override { return OutlinerSelectionChanged; }
+	virtual FOnActorsSelected& OnActorsSelectedExternally() override { return ActorsSelectedExternally; }
 	virtual TSharedPtr<FTabManager> GetTabManager() const override { return TabManager; }
 
 	virtual TSharedPtr<FWorkspaceItem> GetWorkspaceMenuGroup() const override { return WorkspaceItem; }
@@ -40,4 +43,6 @@ private:
 
 	FOnActiveRootActorChanged RootActorChanged;
 	FOnDetailObjectsChanged DetailObjectsChanged;
+	FOnActorsSelected ActorsSelectedExternally;
+	FOnOutlinerSelectionChanged OutlinerSelectionChanged;
 };
\ No newline at end of file
diff --git a/Engine/Plugins/Runtime/nDisplay/Source/DisplayClusterOperator/Public/IDisplayClusterOperatorViewModel.h b/Engine/Plugins/Runtime/nDisplay/Source/DisplayClusterOperator/Public/IDisplayClusterOperatorViewModel.h
--- a/Engine/Plugins/Runtime/nDisplay/Source/DisplayClusterOperator/Public/IDisplayClusterOperatorViewModel.h
+++ b/Engine/Plugins/Runtime/nDisplay/Source/DisplayClusterOperator/Public/IDisplayClusterOperatorViewModel.h
@@ -4,6 +4,7 @@
 
 #include "CoreMinimal.h"
 
+class AActor;
 class ADisplayClusterRootActor;
 class FTabManager;
 class FWorkspaceItem;
@@ -15,6 +16,8 @@ class DISPLAYCLUSTEROPERATOR_API IDisplayClusterOperatorViewModel
 public:
 	DECLARE_MULTICAST_DELEGATE_OneParam(FOnActiveRootActorChanged, ADisplayClusterRootActor*);
 	DECLARE_MULTICAST_DELEGATE_OneParam(FOnDetailObjectsChanged, const TArray<UObject*>&);
+	DECLARE_MULTICAST_DELEGATE_OneParam(FOnOutlinerSelectionChanged, const TArray<AActor*>&);
+	DECLARE_MULTICAST_DELEGATE_TwoParams(FOnActorsSelected, const TArray<AActor*>&, bool bShouldSelect);
 
 public:
 	/** Gets whether the view model has been populated with a valid root actor */
@@ -38,9 +41,18 @@ public:
 	/** Displays the properties of the specified object in the operator's details panel */
 	virtual void ShowDetailsForObjects(const TArray<UObject*>& Objects) = 0;
 
+	/** Selects the specified actors in the operator's outliner panel, or if bShouldSelect is false, deselects them */
+	virtual void SelectActors(const TArray<AActor*>& Actors, bool bShouldSelect) = 0;
+
 	/** Gets the event handler that is raised when the objects being displayed in the operator's details panel are changed */
 	virtual FOnDetailObjectsChanged& OnDetailObjectsChanged() = 0;
 
+	/** Gets the event handler that is raised when the actors selected in the operator's outliner panel are changed */
+	virtual FOnOutlinerSelectionChanged& OnOutlinerSelectionChanged() = 0;
+
+	/** Gets the event handler that is raised when actors are selected from outside the operator itself, i.e. directly from this view model */
+	virtual FOnActorsSelected& OnActorsSelectedExternally() = 0;
+
 	/** Gets the tab manager of the active operator panel, if there is an open operator panel */
 	virtual TSharedPtr<FTabManager> GetTabManager() const = 0;
 
