Fix compilation error on Mac Arm64

diff --git a/Engine/Plugins/Editor/ObjectMixer/ObjectMixer/Source/ObjectMixer/Public/Views/List/Modes/ObjectMixerOutlinerMode.h b/Engine/Plugins/Editor/ObjectMixer/ObjectMixer/Source/ObjectMixer/Public/Views/List/Modes/ObjectMixerOutlinerMode.h
--- a/Engine/Plugins/Editor/ObjectMixer/ObjectMixer/Source/ObjectMixer/Public/Views/List/Modes/ObjectMixerOutlinerMode.h
+++ b/Engine/Plugins/Editor/ObjectMixer/ObjectMixer/Source/ObjectMixer/Public/Views/List/Modes/ObjectMixerOutlinerMode.h
@@ -36,6 +36,8 @@ namespace ObjectMixerOutliner
 
 	struct FComponentSelector : FHybridRowSelector
 	{
+		using FHybridRowSelector::FHybridRowSelector;
+
 		bool operator()(const TWeakPtr<ISceneOutlinerTreeItem>& Item, UActorComponent*& DataOut) const;
 	};
 
@@ -48,6 +50,8 @@ namespace ObjectMixerOutliner
 	/** Functor which can be used to get actors from a selection including component parents */
 	struct FActorSelector : FHybridRowSelector
 	{
+		using FHybridRowSelector::FHybridRowSelector;
+
 		bool operator()(const TWeakPtr<ISceneOutlinerTreeItem>& Item, AActor*& ActorPtrOut) const;
 	};
 
