[ColorGradingEditor] Fix crash and lost selection when changing settings on an ICVFX actors/components in the Color Grading panel

diff --git a/Engine/Plugins/Editor/ObjectMixer/ObjectMixer/Source/ObjectMixer/Private/Views/List/Modes/ObjectMixerOutlinerMode.cpp b/Engine/Plugins/Editor/ObjectMixer/ObjectMixer/Source/ObjectMixer/Private/Views/List/Modes/ObjectMixerOutlinerMode.cpp
--- a/Engine/Plugins/Editor/ObjectMixer/ObjectMixer/Source/ObjectMixer/Private/Views/List/Modes/ObjectMixerOutlinerMode.cpp
+++ b/Engine/Plugins/Editor/ObjectMixer/ObjectMixer/Source/ObjectMixer/Private/Views/List/Modes/ObjectMixerOutlinerMode.cpp
@@ -934,6 +934,12 @@ void FObjectMixerOutlinerMode::OnObjectsReplaced(const TMap<UObject*, UObject*>&
 				SceneOutliner->OnItemLabelChanged(Item);
 			}
 		}
+
+		const int32 SelectionIndex = SelectedIDs.IndexOfByKey(Pair.Key);
+		if (SelectionIndex != INDEX_NONE)
+		{
+			SelectedIDs[SelectionIndex] = Pair.Value;
+		}
 	}
 }
 
@@ -1967,6 +1973,15 @@ void FObjectMixerOutlinerMode::OnItemSelectionChanged(FSceneOutlinerTreeItemPtr
 	{
 		SynchronizeAllSelectionsToEditor();
 	}
+
+	SelectedIDs.Empty();
+
+	Algo::TransformIf(
+		SceneOutliner->GetSelection().SelectedItems,
+		SelectedIDs,
+		[](const TWeakPtr<ISceneOutlinerTreeItem>& Item) -> bool { return Item.IsValid(); },
+		[](const TWeakPtr<ISceneOutlinerTreeItem>& Item) -> FSceneOutlinerTreeItemID { return Item.Pin()->GetID(); }
+	);
 	
 	bShouldTemporarilyForceSelectionSyncToEditor = false;
 }
@@ -3306,9 +3321,22 @@ void FObjectMixerOutlinerMode::SynchronizeSelection()
 	{
 		SynchronizeComponentAndActorSelection();
 		SynchronizeSelectedActorDescs();
+	}
+	else
+	{
+		SceneOutliner->ClearSelection();
 
-		StaticCast<SObjectMixerEditorList*>(SceneOutliner)->GetOnSelectionSynchronized().Broadcast();
+		for (const FSceneOutlinerTreeItemID& TreeItemID : SelectedIDs)
+		{
+			const FSceneOutlinerTreeItemPtr Item = SceneOutliner->GetTreeItem(TreeItemID);
+			if (Item.IsValid())
+			{
+				SceneOutliner->SetItemSelection(Item, true);
+			}
+		}
 	}
+
+	StaticCast<SObjectMixerEditorList*>(SceneOutliner)->GetOnSelectionSynchronized().Broadcast();
 	
 	bShouldTemporarilyForceSelectionSyncFromEditor = false;
 }
diff --git a/Engine/Plugins/Editor/ObjectMixer/ObjectMixer/Source/ObjectMixer/Public/Views/List/Modes/ObjectMixerOutlinerMode.h b/Engine/Plugins/Editor/ObjectMixer/ObjectMixer/Source/ObjectMixer/Public/Views/List/Modes/ObjectMixerOutlinerMode.h
--- a/Engine/Plugins/Editor/ObjectMixer/ObjectMixer/Source/ObjectMixer/Public/Views/List/Modes/ObjectMixerOutlinerMode.h
+++ b/Engine/Plugins/Editor/ObjectMixer/ObjectMixer/Source/ObjectMixer/Public/Views/List/Modes/ObjectMixerOutlinerMode.h
@@ -5,6 +5,7 @@
 #include "ActorMode.h"
 #include "Containers/Map.h"
 #include "EditorConfigBase.h"
+#include "ISceneOutlinerTreeItem.h"
 #include "ObjectFilter/ObjectMixerEditorObjectFilter.h"
 #include "UObject/ObjectPtr.h"
 #include "UObject/UObjectGlobals.h"
@@ -387,4 +388,7 @@ protected:
 	
 	bool bShouldTemporarilyForceSelectionSyncFromEditor = false;
 	bool bShouldTemporarilyForceSelectionSyncToEditor = false;
+
+	/** The last selected tree item IDs, stored here so we can maintain the selection across refreshes when not synced with the editor. */
+	TArray<FSceneOutlinerTreeItemID> SelectedIDs;
 };
diff --git a/Engine/Plugins/Runtime/nDisplay/Source/DisplayClusterColorGrading/Private/DataModelGenerators/DisplayClusterColorGradingGenerator_RootActor.cpp b/Engine/Plugins/Runtime/nDisplay/Source/DisplayClusterColorGrading/Private/DataModelGenerators/DisplayClusterColorGradingGenerator_RootActor.cpp
--- a/Engine/Plugins/Runtime/nDisplay/Source/DisplayClusterColorGrading/Private/DataModelGenerators/DisplayClusterColorGradingGenerator_RootActor.cpp
+++ b/Engine/Plugins/Runtime/nDisplay/Source/DisplayClusterColorGrading/Private/DataModelGenerators/DisplayClusterColorGradingGenerator_RootActor.cpp
@@ -136,6 +136,11 @@ TSharedPtr<IPropertyHandle> MakePropertyTransactional(const TSharedPtr<IProperty
 			PropertyHandle->GetOuterObjects(OuterObjects);
 			for (UObject* Object : OuterObjects)
 			{
+				if (!Object)
+				{
+					continue;
+				}
+
 				if (!Object->HasAnyFlags(RF_Transactional))
 				{
 					Object->SetFlags(RF_Transactional);
