[ObjectMixer] Notify editor of selection when creating a new actor so it's immediately displayed in the Details panel

diff --git a/Engine/Plugins/Editor/ObjectMixer/ObjectMixer/Source/ObjectMixer/Private/Views/Widgets/SObjectMixerPlacementAssetMenuEntry.cpp b/Engine/Plugins/Editor/ObjectMixer/ObjectMixer/Source/ObjectMixer/Private/Views/Widgets/SObjectMixerPlacementAssetMenuEntry.cpp
--- a/Engine/Plugins/Editor/ObjectMixer/ObjectMixer/Source/ObjectMixer/Private/Views/Widgets/SObjectMixerPlacementAssetMenuEntry.cpp
+++ b/Engine/Plugins/Editor/ObjectMixer/ObjectMixer/Source/ObjectMixer/Private/Views/Widgets/SObjectMixerPlacementAssetMenuEntry.cpp
@@ -186,7 +186,7 @@ FReply SObjectMixerPlacementAssetMenuEntry::OnMouseButtonUp(const FGeometry& MyG
 		}
 
 		GEditor->SelectNone(true, true, false);
-		GEditor->SelectActor(NewActor, true, false);
+		GEditor->SelectActor(NewActor, true, true);
 
 		return FReply::Handled();
 	}
