PR #10580: Adding missing dependency

diff --git a/Engine/Plugins/Editor/ObjectMixer/ObjectMixer/Source/ObjectMixer/ObjectMixerEditor.build.cs b/Engine/Plugins/Editor/ObjectMixer/ObjectMixer/Source/ObjectMixer/ObjectMixerEditor.build.cs
--- a/Engine/Plugins/Editor/ObjectMixer/ObjectMixer/Source/ObjectMixer/ObjectMixerEditor.build.cs
+++ b/Engine/Plugins/Editor/ObjectMixer/ObjectMixer/Source/ObjectMixer/ObjectMixerEditor.build.cs
@@ -54,7 +54,8 @@ public class ObjectMixerEditor : ModuleRules
 				"TypedElementRuntime",
 				"UMG",
 				"UnrealEd",
-				"WorkspaceMenuStructure"
+				"WorkspaceMenuStructure",
+				"SourceControl"
 			}
 		);
 	}
