[ColorGrading] [ColorCorrectRegions] Add Color Grading Mixer support for per-actor color correction

diff --git a/Engine/Plugins/Editor/ColorGrading/Source/ColorGradingEditor/Private/ColorGradingEditorDataModel.cpp b/Engine/Plugins/Editor/ColorGrading/Source/ColorGradingEditor/Private/ColorGradingEditorDataModel.cpp
--- a/Engine/Plugins/Editor/ColorGrading/Source/ColorGradingEditor/Private/ColorGradingEditorDataModel.cpp
+++ b/Engine/Plugins/Editor/ColorGrading/Source/ColorGradingEditor/Private/ColorGradingEditorDataModel.cpp
@@ -171,16 +171,16 @@ void FColorGradingEditorDataModel::SetPanelState(const FColorGradingPanelState&
 	SelectedColorGradingGroupIndex = InPanelState.SelectedColorGradingGroup;
 	SelectedColorGradingElementIndex = InPanelState.SelectedColorGradingElement;
 
-	TArray<UObject*> ObjectsToSelect;
-	for (const TWeakObjectPtr<UObject>& Object : InPanelState.SelectedObjects)
+	TArray<UObject*> ObjectsToControl;
+	for (const TWeakObjectPtr<UObject>& Object : InPanelState.ControlledObjects)
 	{
 		if (Object.IsValid())
 		{
-			ObjectsToSelect.Add(Object.Get());
+			ObjectsToControl.Add(Object.Get());
 		}
 	}
 
-	for (const UObject* Object : ObjectsToSelect)
+	for (const UObject* Object : ObjectsToControl)
 	{
 		if (Object)
 		{
@@ -190,7 +190,7 @@ void FColorGradingEditorDataModel::SetPanelState(const FColorGradingPanelState&
 
 	if (PropertyRowGenerator.IsValid())
 	{
-		PropertyRowGenerator->SetObjects(ObjectsToSelect);
+		PropertyRowGenerator->SetObjects(ObjectsToControl);
 	}
 
 	// After the data model has been created as part of the SetObjects call, check that the saved selected color grading group is still valid,
diff --git a/Engine/Plugins/Editor/ColorGrading/Source/ColorGradingEditor/Private/ColorGradingMixerObjectFilter.cpp b/Engine/Plugins/Editor/ColorGrading/Source/ColorGradingEditor/Private/ColorGradingMixerObjectFilter.cpp
--- a/Engine/Plugins/Editor/ColorGrading/Source/ColorGradingEditor/Private/ColorGradingMixerObjectFilter.cpp
+++ b/Engine/Plugins/Editor/ColorGrading/Source/ColorGradingEditor/Private/ColorGradingMixerObjectFilter.cpp
@@ -1,7 +1,11 @@
 // Copyright Epic Games, Inc. All Rights Reserved.
 
 #include "ColorGradingMixerObjectFilter.h"
+
+#include "ActorTreeItem.h"
+#include "ComponentTreeItem.h"
 #include "ColorGradingMixerObjectFilterRegistry.h"
+#include "ISceneOutliner.h"
 
 TSet<UClass*> UColorGradingMixerObjectFilter::GetObjectClassesToFilter() const
 {
@@ -11,4 +15,89 @@ TSet<UClass*> UColorGradingMixerObjectFilter::GetObjectClassesToFilter() const
 TSet<TSubclassOf<AActor>> UColorGradingMixerObjectFilter::GetObjectClassesToPlace() const
 {
 	return FColorGradingMixerObjectFilterRegistry::GetActorClassesToPlace();
-}
\ No newline at end of file
+}
+
+TArray<AActor*> UColorGradingMixerObjectFilter::FindAssociatedActors(AActor* InActor) const
+{
+	if (InActor)
+	{
+		if (const IColorGradingMixerObjectHierarchyConfig* Config = FColorGradingMixerObjectFilterRegistry::GetClassHierarchyConfig(InActor->GetClass()))
+		{
+			return Config->FindAssociatedActors(InActor);
+		}
+	}
+
+	return {};
+}
+
+
+bool UColorGradingMixerObjectFilter::IsActorAssociated(AActor* Actor, AActor* AssociatedActor) const
+{
+	if (Actor && AssociatedActor)
+	{
+		if (const IColorGradingMixerObjectHierarchyConfig* Config = FColorGradingMixerObjectFilterRegistry::GetClassHierarchyConfig(Actor->GetClass()))
+		{
+			return Config->IsActorAssociated(Actor, AssociatedActor);
+		}
+	}
+
+	return false;
+}
+
+bool UColorGradingMixerObjectFilter::HasCustomDropHandling(const ISceneOutlinerTreeItem& DropTarget) const
+{
+	if (const UObject* TargetObject = GetObjectForTreeItem(DropTarget))
+	{
+		if (const IColorGradingMixerObjectHierarchyConfig* Config = FColorGradingMixerObjectFilterRegistry::GetClassHierarchyConfig(TargetObject->GetClass()))
+		{
+			return Config->HasCustomDropHandling();
+		}
+	}
+
+	return false;
+}
+
+FSceneOutlinerDragValidationInfo UColorGradingMixerObjectFilter::ValidateDrop(const ISceneOutlinerTreeItem& DropTarget, const FSceneOutlinerDragDropPayload& Payload) const
+{
+	if (UObject* TargetObject = GetObjectForTreeItem(DropTarget))
+	{
+		if (const IColorGradingMixerObjectHierarchyConfig* Config = FColorGradingMixerObjectFilterRegistry::GetClassHierarchyConfig(TargetObject->GetClass()))
+		{
+			return Config->ValidateDrop(TargetObject, Payload);
+		}
+	}
+
+	return FSceneOutlinerDragValidationInfo::Invalid();
+}
+
+void UColorGradingMixerObjectFilter::OnDrop(ISceneOutlinerTreeItem& DropTarget, const FSceneOutlinerDragDropPayload& Payload, const FSceneOutlinerDragValidationInfo& ValidationInfo) const
+{
+	if (UObject* TargetObject = GetObjectForTreeItem(DropTarget))
+	{
+		if (const IColorGradingMixerObjectHierarchyConfig* Config = FColorGradingMixerObjectFilterRegistry::GetClassHierarchyConfig(TargetObject->GetClass()))
+		{
+			if (Config->OnDrop(TargetObject, Payload, ValidationInfo))
+			{
+				if (ISceneOutliner* SceneOutliner = DropTarget.WeakSceneOutliner.Pin().Get())
+				{
+					SceneOutliner->FullRefresh();
+				}
+			}
+		}
+	}
+}
+
+UObject* UColorGradingMixerObjectFilter::GetObjectForTreeItem(const ISceneOutlinerTreeItem& TreeItem) const
+{
+	if (const FActorTreeItem* ActorTreeItem = TreeItem.CastTo<FActorTreeItem>())
+	{
+		return ActorTreeItem->Actor.Get();
+	}
+
+	if (const FComponentTreeItem* ComponentTreeItem = TreeItem.CastTo<FComponentTreeItem>())
+	{
+		return ComponentTreeItem->Component.Get();
+	}
+
+	return nullptr;
+}
diff --git a/Engine/Plugins/Editor/ColorGrading/Source/ColorGradingEditor/Private/ColorGradingMixerObjectFilter.h b/Engine/Plugins/Editor/ColorGrading/Source/ColorGradingEditor/Private/ColorGradingMixerObjectFilter.h
--- a/Engine/Plugins/Editor/ColorGrading/Source/ColorGradingEditor/Private/ColorGradingMixerObjectFilter.h
+++ b/Engine/Plugins/Editor/ColorGrading/Source/ColorGradingEditor/Private/ColorGradingMixerObjectFilter.h
@@ -14,6 +14,11 @@ public:
 
 	virtual TSet<UClass*> GetObjectClassesToFilter() const override;
 	virtual TSet<TSubclassOf<AActor>> GetObjectClassesToPlace() const override;
+	virtual TArray<AActor*> FindAssociatedActors(AActor* InActor) const override;
+	virtual bool IsActorAssociated(AActor* Actor, AActor* AssociatedActor) const override;
+	virtual bool HasCustomDropHandling(const ISceneOutlinerTreeItem& DropTarget) const override;
+	virtual FSceneOutlinerDragValidationInfo ValidateDrop(const ISceneOutlinerTreeItem& DropTarget, const FSceneOutlinerDragDropPayload& Payload) const override;
+	virtual void OnDrop(ISceneOutlinerTreeItem& DropTarget, const FSceneOutlinerDragDropPayload& Payload, const FSceneOutlinerDragValidationInfo& ValidationInfo) const override;
 
 	virtual bool GetShowTransientObjects() const override
 	{
@@ -62,4 +67,8 @@ public:
 	{
 		return EObjectMixerInheritanceInclusionOptions::None;
 	}
+
+private:
+	/** Given a tree item, get the object it represents */
+	UObject* GetObjectForTreeItem(const ISceneOutlinerTreeItem& TreeItem) const;
 };
diff --git a/Engine/Plugins/Editor/ColorGrading/Source/ColorGradingEditor/Private/ColorGradingMixerObjectFilterRegistry.cpp b/Engine/Plugins/Editor/ColorGrading/Source/ColorGradingEditor/Private/ColorGradingMixerObjectFilterRegistry.cpp
new file mode 100644
--- /dev/null
+++ b/Engine/Plugins/Editor/ColorGrading/Source/ColorGradingEditor/Private/ColorGradingMixerObjectFilterRegistry.cpp
@@ -0,0 +1,29 @@
+// Copyright Epic Games, Inc. All Rights Reserved.
+
+#include "ColorGradingMixerObjectFilterRegistry.h"
+
+const IColorGradingMixerObjectHierarchyConfig* FColorGradingMixerObjectFilterRegistry::GetClassHierarchyConfig(TSubclassOf<AActor> Class)
+{
+	if (TSharedPtr<IColorGradingMixerObjectHierarchyConfig>* Config = HierarchyConfigInstances.Find(Class))
+	{
+		return Config->Get();
+	}
+
+	// Find the config for this class or its ancestors
+	while (Class)
+	{
+		if (FGetObjectHierarchyConfig* ConfigDelegate = HierarchyConfigs.Find(Class))
+		{
+			TSharedRef<IColorGradingMixerObjectHierarchyConfig> Config = ConfigDelegate->Execute();
+			HierarchyConfigInstances.Emplace(Class, Config);
+
+			return Config.ToSharedPtr().Get();
+		}
+
+		Class = Class->GetSuperClass();
+	}
+
+	// No config found. Cache nullptr so we don't have to look this up again
+	HierarchyConfigInstances.Add(Class, nullptr);
+	return nullptr;
+}
\ No newline at end of file
diff --git a/Engine/Plugins/Editor/ColorGrading/Source/ColorGradingEditor/Private/SColorGradingPanel.cpp b/Engine/Plugins/Editor/ColorGrading/Source/ColorGradingEditor/Private/SColorGradingPanel.cpp
--- a/Engine/Plugins/Editor/ColorGrading/Source/ColorGradingEditor/Private/SColorGradingPanel.cpp
+++ b/Engine/Plugins/Editor/ColorGrading/Source/ColorGradingEditor/Private/SColorGradingPanel.cpp
@@ -5,6 +5,7 @@
 #include "ColorGradingCommands.h"
 #include "ColorGradingEditorDataModel.h"
 #include "ColorGradingMixerObjectFilter.h"
+#include "ColorGradingMixerObjectFilterRegistry.h"
 #include "IColorGradingEditor.h"
 #include "SColorGradingColorWheelPanel.h"
 
@@ -19,6 +20,7 @@
 #include "Framework/MultiBox/MultiBoxBuilder.h"
 #include "LevelEditorSubsystem.h"
 #include "Views/List/ObjectMixerEditorList.h"
+#include "Views/List/RowTypes/ObjectMixerEditorListRowActor.h"
 #include "Views/List/SObjectMixerEditorList.h"
 #include "Widgets/Input/SButton.h"
 #include "Widgets/Layout/SBox.h"
@@ -213,32 +215,13 @@ void SColorGradingPanel::GetPanelState(FColorGradingPanelState& OutPanelState) c
 	if (ColorGradingObjectListView.IsValid())
 	{
 		TArray<FSceneOutlinerTreeItemPtr> SelectedItems = ColorGradingObjectListView->GetSelectedItems();
+		TArray<UObject*> SelectedObjects;
+		TArray<UObject*> ControlledObjects;
 
-		for (const FSceneOutlinerTreeItemPtr& SelectedItem : SelectedItems)
-		{
-			if (!SelectedItem.IsValid())
-			{
-				continue;
-			}
-			
-			if (const FComponentTreeItem* ComponentTreeItem = SelectedItem->CastTo<FComponentTreeItem>())
-			{
-				if (ComponentTreeItem->Component.IsValid())
-				{
-					OutPanelState.SelectedObjects.Add(ComponentTreeItem->Component.Get());
-					continue;
-				}
-			}
-			
-			if (const FActorTreeItem* ActorTreeItem = SelectedItem->CastTo<FActorTreeItem>())
-			{
-				if (ActorTreeItem->Actor.IsValid())
-				{
-					OutPanelState.SelectedObjects.Add(ActorTreeItem->Actor.Get());
-					continue;
-				}
-			}
-		}
+		GetSelectedAndControlledObjects(SelectedItems, SelectedObjects, ControlledObjects);
+
+		OutPanelState.SelectedObjects.Append(SelectedObjects);
+		OutPanelState.ControlledObjects.Append(ControlledObjects);
 	}
 }
 
@@ -259,7 +242,6 @@ void SColorGradingPanel::SetPanelState(const FColorGradingPanelState& InPanelSta
 			ItemsToSelect.Add(Item);
 			break;
 		}
-
 	}
 
 	if (!ItemsToSelect.IsEmpty() && ColorGradingObjectListView.IsValid())
@@ -276,15 +258,20 @@ void SColorGradingPanel::SetPanelState(const FColorGradingPanelState& InPanelSta
 	}
 }
 
-void SColorGradingPanel::SetSelectedObjects(const TArray<UObject*>& Objects)
+void SColorGradingPanel::SetSelectedObjects(const TArray<UObject*>& SelectedObjects, const TArray<UObject*>* ControlledObjects)
 {
-	ColorGradingDataModel->SetObjects(Objects);
+	const TArray<UObject*>& ControlledObjectsToUse = ControlledObjects ? *ControlledObjects : SelectedObjects;
+
+	ColorGradingDataModel->SetObjects(ControlledObjectsToUse);
 
 	FColorGradingPanelState PanelState;
 	GetPanelState(PanelState);
 
-	PanelState.SelectedObjects.Empty(Objects.Num());
-	PanelState.SelectedObjects.Append(Objects);
+	PanelState.SelectedObjects.Empty(SelectedObjects.Num());
+	PanelState.SelectedObjects.Append(SelectedObjects);
+
+	PanelState.ControlledObjects.Empty(ControlledObjectsToUse.Num());
+	PanelState.ControlledObjects.Append(ControlledObjectsToUse);
 
 	SetPanelState(PanelState);
 }
@@ -546,31 +533,67 @@ void SColorGradingPanel::UpdateSelectionFromList()
 {
 	TArray<FSceneOutlinerTreeItemPtr> SelectedOutlinerItems = ColorGradingObjectListView->GetSelectedItems();
 	TArray<UObject*> SelectedObjects;
+	TArray<UObject*> ControlledObjects;
+
+	GetSelectedAndControlledObjects(SelectedOutlinerItems, SelectedObjects, ControlledObjects);
 
-	for (FSceneOutlinerTreeItemPtr OutlinerItems : SelectedOutlinerItems)
+	TArray<TWeakObjectPtr<UObject>> OldControlledObjects = ColorGradingDataModel->GetObjects();
+	if (!Algo::Compare(OldControlledObjects, ControlledObjects))
 	{
-		if (const FActorTreeItem* ActorTreeItem = OutlinerItems->CastTo<FActorTreeItem>())
+		SetSelectedObjects(SelectedObjects, &ControlledObjects);
+	}
+}
+
+void SColorGradingPanel::GetSelectedAndControlledObjects(const TArray<FSceneOutlinerTreeItemPtr>& InSelectedItems, TArray<UObject*>& OutSelectedObjects, TArray<UObject*>& OutControlledObjects) const
+{
+	OutSelectedObjects.Reserve(InSelectedItems.Num());
+	OutControlledObjects.Reserve(InSelectedItems.Num());
+
+	for (FSceneOutlinerTreeItemPtr TreeItem : InSelectedItems)
+	{
+		if (!TreeItem.IsValid())
+		{
+			continue;
+		}
+
+		if (const FActorTreeItem* ActorTreeItem = TreeItem->CastTo<FActorTreeItem>())
 		{
-			if (ActorTreeItem->Actor.IsValid())
+			if (AActor* Actor = ActorTreeItem->Actor.Get())
 			{
-				SelectedObjects.Add(ActorTreeItem->Actor.Get());
+				// Check if this actor is listed as an associated actor, in which case its parent is the actual selection target for color grading
+				if (const FObjectMixerEditorListRowActor* ObjectMixerActorItem = TreeItem->CastTo<FObjectMixerEditorListRowActor>())
+				{
+					if (AActor* ParentActor = ObjectMixerActorItem->OverrideParent.Get())
+					{
+						if (const IColorGradingMixerObjectHierarchyConfig* HierarchyConfig = FColorGradingMixerObjectFilterRegistry::GetClassHierarchyConfig(ParentActor->GetClass()))
+						{
+							if (HierarchyConfig->IsActorAssociated(ParentActor, Actor))
+							{
+								OutControlledObjects.Add(ParentActor);
+								OutSelectedObjects.Add(Actor);
+								continue;
+							}
+						}
+					}
+				}
+
+				// Otherwise we want to control it directly
+				OutControlledObjects.Add(Actor);
+				OutSelectedObjects.Add(Actor);
+				continue;
 			}
 		}
 
-		if (const FComponentTreeItem* ComponentTreeItem = OutlinerItems->CastTo<FComponentTreeItem>())
+		if (const FComponentTreeItem* ComponentTreeItem = TreeItem->CastTo<FComponentTreeItem>())
 		{
-			if (ComponentTreeItem->Component.IsValid())
+			if (UActorComponent* Component = ComponentTreeItem->Component.Get())
 			{
-				SelectedObjects.Add(ComponentTreeItem->Component.Get());
+				OutControlledObjects.Add(Component);
+				OutSelectedObjects.Add(Component);
+				continue;
 			}
 		}
 	}
-
-	TArray<TWeakObjectPtr<UObject>> OldSelectedObjects = ColorGradingDataModel->GetObjects();
-	if (!Algo::Compare(OldSelectedObjects, SelectedObjects))
-	{
-		SetSelectedObjects(SelectedObjects);
-	}
 }
 
 FReply SColorGradingPanel::DockInLayout()
diff --git a/Engine/Plugins/Editor/ColorGrading/Source/ColorGradingEditor/Public/ColorGradingMixerObjectFilterRegistry.h b/Engine/Plugins/Editor/ColorGrading/Source/ColorGradingEditor/Public/ColorGradingMixerObjectFilterRegistry.h
--- a/Engine/Plugins/Editor/ColorGrading/Source/ColorGradingEditor/Public/ColorGradingMixerObjectFilterRegistry.h
+++ b/Engine/Plugins/Editor/ColorGrading/Source/ColorGradingEditor/Public/ColorGradingMixerObjectFilterRegistry.h
@@ -4,10 +4,42 @@
 
 #include "GameFramework/Actor.h"
 #include "HAL/Platform.h"
+#include "SceneOutlinerDragDrop.h"
 #include "Templates/SubclassOf.h"
 
 class UClass;
 
+/** Interface to configure how an object is displayed in the Color Grading panel's outliner */
+struct COLORGRADINGEDITOR_API IColorGradingMixerObjectHierarchyConfig : public TSharedFromThis<IColorGradingMixerObjectHierarchyConfig>
+{
+public:
+	virtual ~IColorGradingMixerObjectHierarchyConfig() {}
+
+	/** Generate a list of additional actors to include as children in the ObjectMixer hierarchy for the given object */
+	virtual TArray<AActor*> FindAssociatedActors(UObject* ParentObject) const { return {}; }
+
+	/**
+	 * Return true if AssociatedActor is associated with ParentObject.
+	 * This should return true if the actor would be included in FindAssociatedActors, but is a separate function to avoid allocating
+	 * an array every time.
+	 */
+	virtual bool IsActorAssociated(UObject* ParentObject, AActor* AssociatedActor) const { return false; }
+
+	/** 
+	 * If true, this has custom handling for drop operations and will override the hierarchy's default behaviour with the ValidateDrop
+	 * and OnDrop functions.
+	 */
+	virtual bool HasCustomDropHandling() const { return false; }
+
+	/** Test whether the specified payload can be dropped onto a tree item representing this object */
+	virtual FSceneOutlinerDragValidationInfo ValidateDrop(UObject* DropTarget, const FSceneOutlinerDragDropPayload& Payload) const { return FSceneOutlinerDragValidationInfo::Invalid(); }
+
+	/** Called when a payload is dropped onto a target. If this returns true, the outliner will be refreshed after the operation. */
+	virtual bool OnDrop(UObject* DropTarget, const FSceneOutlinerDragDropPayload& Payload, const FSceneOutlinerDragValidationInfo& ValidationInfo) const { return false;  }
+};
+
+DECLARE_DELEGATE_RetVal(TSharedRef<IColorGradingMixerObjectHierarchyConfig>, FGetObjectHierarchyConfig);
+
 /**
  * Contains functions for controlling which classes can be managed from the Color Grading panel's ObjectMixer-based hierarchy panel
  */
@@ -15,10 +47,19 @@ struct COLORGRADINGEDITOR_API FColorGradingMixerObjectFilterRegistry
 {
 public:
 
-	/** Register an object class that can be seen in a Color Grading panel's object list */
-	static void RegisterObjectClassToFilter(UClass* Class)
+	/**
+	 * Register an object class that can be seen in a Color Grading panel's object list.
+	 * @param Class The object class to register
+	 * @param Config Optional delegate that creates a configuration for how the object will be displayed in the hierarchy
+	 */
+	static void RegisterObjectClassToFilter(UClass* Class, FGetObjectHierarchyConfig CreateConfigDelegate = nullptr)
 	{
 		ObjectClassesToFilter.Add(Class);
+
+		if (CreateConfigDelegate.IsBound())
+		{
+			HierarchyConfigs.Add(Class, CreateConfigDelegate);
+		}
 	}
 
 	/** Register an actor class that can be placed from the Color Grading panel's object list */
@@ -33,10 +74,19 @@ public:
 	/** Get the set of actor classes that can be placed from a Color Grading panel's object list */
 	static const TSet<TSubclassOf<AActor>>& GetActorClassesToPlace() { return ActorClassesToPlace; }
 
+	/** Get the hierarchy configuration for a class, or null if none was provided */
+	static const IColorGradingMixerObjectHierarchyConfig* GetClassHierarchyConfig(TSubclassOf<AActor> Class);
+
 private:
 	/** Set of classes that can be seen in the object panel */
 	inline static TSet<UClass*> ObjectClassesToFilter;
 
 	/** Set of classes that can be placed from the object panel */
 	inline static TSet<TSubclassOf<AActor>> ActorClassesToPlace;
+
+	/** Map from object class to hierarchy configuration, if provided */
+	inline static TMap<TWeakObjectPtr<UClass>, FGetObjectHierarchyConfig> HierarchyConfigs;
+
+	/** Map from object class to hierarchy configuration instance. If null, none exist for the class */
+	inline static TMap<TWeakObjectPtr<UClass>, TSharedPtr<IColorGradingMixerObjectHierarchyConfig>> HierarchyConfigInstances;
 };
diff --git a/Engine/Plugins/Editor/ColorGrading/Source/ColorGradingEditor/Public/ColorGradingPanelState.h b/Engine/Plugins/Editor/ColorGrading/Source/ColorGradingEditor/Public/ColorGradingPanelState.h
--- a/Engine/Plugins/Editor/ColorGrading/Source/ColorGradingEditor/Public/ColorGradingPanelState.h
+++ b/Engine/Plugins/Editor/ColorGrading/Source/ColorGradingEditor/Public/ColorGradingPanelState.h
@@ -15,6 +15,13 @@ struct COLORGRADINGEDITOR_API FColorGradingPanelState
 	/** The objects that are selected in the list */
 	TArray<TWeakObjectPtr<UObject>> SelectedObjects;
 
+	/**
+	 * The objects that are being controlled by the color grading controls.
+	 * This can differ from SelectedObjects when SelectedObjects contains objects that are associated
+	 * with a color-gradable object, but don't have controls themselves.
+	 */
+	TArray<TWeakObjectPtr<UObject>> ControlledObjects;
+
 	/** The color grading group that is selected */
 	int32 SelectedColorGradingGroup = INDEX_NONE;
 
diff --git a/Engine/Plugins/Editor/ColorGrading/Source/ColorGradingEditor/Public/SColorGradingPanel.h b/Engine/Plugins/Editor/ColorGrading/Source/ColorGradingEditor/Public/SColorGradingPanel.h
--- a/Engine/Plugins/Editor/ColorGrading/Source/ColorGradingEditor/Public/SColorGradingPanel.h
+++ b/Engine/Plugins/Editor/ColorGrading/Source/ColorGradingEditor/Public/SColorGradingPanel.h
@@ -57,8 +57,12 @@ public:
 	/** Sets the state of the panel UI */
 	void SetPanelState(const FColorGradingPanelState& InPanelState);
 
-	/** Set the list of selected objects, updating state and data model as appropriate */
-	void SetSelectedObjects(const TArray<UObject*>& Objects);
+	/**
+	 * Set the list of selected objects, updating state and data model as appropriate.
+	 * If ControlledObjects is provided, only the objects in that array will be used to display the color grading wheels.
+	 * Otherwise, the controlled objects will use the contents of SelectedObjects.
+	 */
+	void SetSelectedObjects(const TArray<UObject*>& SelectedObjects, const TArray<UObject*>* ControlledObjects = nullptr);
 
 private:
 	/** Creates the button used to dock the drawer in the operator panel */
@@ -115,6 +119,8 @@ private:
 	/** Update the selected items to match the list view */
 	void UpdateSelectionFromList();
 
+	/** Given an array of selected tree items, determine which objects will be selected and/or controlled */
+	void GetSelectedAndControlledObjects(const TArray<FSceneOutlinerTreeItemPtr>& InSelectedItems, TArray<UObject*>& OutSelectedObjects, TArray<UObject*>& OutControlledObjects) const;
 
 private:
 	/** Model for the object mixer list used to display the color gradable object hierarchy. */
diff --git a/Engine/Plugins/Editor/ObjectMixer/ObjectMixer/Source/ObjectMixer/Private/ObjectFilter/ObjectMixerEditorObjectFilter.cpp b/Engine/Plugins/Editor/ObjectMixer/ObjectMixer/Source/ObjectMixer/Private/ObjectFilter/ObjectMixerEditorObjectFilter.cpp
--- a/Engine/Plugins/Editor/ObjectMixer/ObjectMixer/Source/ObjectMixer/Private/ObjectFilter/ObjectMixerEditorObjectFilter.cpp
+++ b/Engine/Plugins/Editor/ObjectMixer/ObjectMixer/Source/ObjectMixer/Private/ObjectFilter/ObjectMixerEditorObjectFilter.cpp
@@ -60,6 +60,30 @@ TSet<FName> UObjectMixerObjectFilter::GetPropertiesThatRequireListRefresh() cons
 	return { "Mobility" };
 }
 
+TArray<AActor*> UObjectMixerObjectFilter::FindAssociatedActors(AActor* InActor) const
+{
+	return {};
+}
+
+bool UObjectMixerObjectFilter::IsActorAssociated(AActor* Actor, AActor* AssociatedActor) const
+{
+	return false;
+}
+
+bool UObjectMixerObjectFilter::HasCustomDropHandling(const ISceneOutlinerTreeItem& DropTarget) const
+{
+	return false;
+}
+
+FSceneOutlinerDragValidationInfo UObjectMixerObjectFilter::ValidateDrop(const ISceneOutlinerTreeItem& DropTarget, const FSceneOutlinerDragDropPayload& Payload) const
+{
+	return FSceneOutlinerDragValidationInfo::Invalid();
+}
+
+void UObjectMixerObjectFilter::OnDrop(ISceneOutlinerTreeItem& DropTarget, const FSceneOutlinerDragDropPayload& Payload, const FSceneOutlinerDragValidationInfo& ValidationInfo) const
+{
+}
+
 TSet<UClass*> UObjectMixerObjectFilter::GetParentAndChildClassesFromSpecifiedClasses(
 	const TSet<UClass*>& InSpecifiedClasses, EObjectMixerInheritanceInclusionOptions Options)
 {
diff --git a/Engine/Plugins/Editor/ObjectMixer/ObjectMixer/Source/ObjectMixer/Private/Views/List/Modes/ObjectMixerOutlinerHierarchy.cpp b/Engine/Plugins/Editor/ObjectMixer/ObjectMixer/Source/ObjectMixer/Private/Views/List/Modes/ObjectMixerOutlinerHierarchy.cpp
--- a/Engine/Plugins/Editor/ObjectMixer/ObjectMixer/Source/ObjectMixer/Private/Views/List/Modes/ObjectMixerOutlinerHierarchy.cpp
+++ b/Engine/Plugins/Editor/ObjectMixer/ObjectMixer/Source/ObjectMixer/Private/Views/List/Modes/ObjectMixerOutlinerHierarchy.cpp
@@ -19,6 +19,7 @@
 #include "EditorActorFolders.h"
 #include "ISceneOutlinerMode.h"
 #include "LevelTreeItem.h"
+#include "ObjectFilter/ObjectMixerEditorObjectFilter.h"
 #include "ObjectMixerEditorSettings.h"
 #include "WorldTreeItem.h"
 #include "Engine/Engine.h"
@@ -85,6 +86,16 @@ FObjectMixerOutlinerMode* FObjectMixerOutlinerHierarchy::GetCastedMode() const
 	return StaticCast<FObjectMixerOutlinerMode*>(Mode);
 }
 
+const UObjectMixerObjectFilter* FObjectMixerOutlinerHierarchy::GetObjectFilter() const
+{
+	if (const TSharedPtr<FObjectMixerEditorList> ListModel = GetCastedMode()->GetListModelPtr().Pin())
+	{
+		return ListModel->GetMainObjectFilterInstance();
+	}
+
+	return nullptr;
+}
+
 FObjectMixerOutlinerHierarchy::FObjectMixerOutlinerHierarchy(ISceneOutlinerMode* Mode, const TWeakObjectPtr<UWorld>& World)
 	: ISceneOutlinerHierarchy(Mode)
 	, RepresentingWorld(World)
@@ -141,6 +152,20 @@ FSceneOutlinerTreeItemPtr FObjectMixerOutlinerHierarchy::FindOrCreateParentItem(
 	}
 	else if (const FActorTreeItem* ActorTreeItem = Item.CastTo<FActorTreeItem>())
 	{
+		// Override parent
+		if (const FObjectMixerEditorListRowActor* ObjectMixerActorItem = Item.CastTo<FObjectMixerEditorListRowActor>())
+		{
+			if (AActor* ParentActor = ObjectMixerActorItem->OverrideParent.Get())
+			{
+				// If OverrideParent is set, assume this item was created in the process of creating the parent, so we
+				// don't have to create the parent as a fallback here
+				if (const FSceneOutlinerTreeItemPtr* ParentItem = Items.Find(ParentActor))
+				{
+					return *ParentItem;
+				}
+			}
+		}
+
 		if (AActor* Actor = ActorTreeItem->Actor.Get())
 		{
 			// Parent Actor (Actor attachement / parenting)
@@ -543,8 +568,10 @@ void FObjectMixerOutlinerHierarchy::CreateWorldChildren(UWorld* World, TArray<FS
 	}
 }
 
-void FObjectMixerOutlinerHierarchy::CreateComponentItems(const AActor* Actor, TArray<FSceneOutlinerTreeItemPtr>& OutItems) const
+bool FObjectMixerOutlinerHierarchy::CreateComponentItems(AActor* Actor, TArray<FSceneOutlinerTreeItemPtr>& OutItems) const
 {
+	bool bAnyItemsBeforeFiltering = false;
+
 	check(Actor);
 	// Add all this actors components if showing components and the owning actor was created
 	if (bShowingComponents)
@@ -554,7 +581,9 @@ void FObjectMixerOutlinerHierarchy::CreateComponentItems(const AActor* Actor, TA
 		for (UActorComponent* Component : ActorComponents)
 		{
 			if (IsValid(Component) && DoesWorldObjectHaveAcceptableClass(Component))
-			{					
+			{
+				bAnyItemsBeforeFiltering = true;
+
 				if (FSceneOutlinerTreeItemPtr ComponentItem =
 					Mode->CreateItemFor<FObjectMixerEditorListRowComponent>(
 						FObjectMixerEditorListRowComponent(Component, GetCastedMode()->GetSceneOutliner())))
@@ -564,41 +593,56 @@ void FObjectMixerOutlinerHierarchy::CreateComponentItems(const AActor* Actor, TA
 			}
 		}
 	}
+
+	return bAnyItemsBeforeFiltering;
 }
 
-TArray<FSceneOutlinerTreeItemPtr> FObjectMixerOutlinerHierarchy::ConditionallyCreateActorAndComponentItems(AActor* Actor) const
+bool FObjectMixerOutlinerHierarchy::CreateAssociatedActorItems(AActor* Actor, TArray<FSceneOutlinerTreeItemPtr>& OutItems) const
 {
-	if (!IsValid(Actor))
-	{
-		return {};
-	}
-	
-	// Whether or not we have components to return, we should create an actor row if components would be returned if there were no filters
-	bool bWouldReturnAnyComponentsBeforeFiltering = false;
-	
-	TArray<FSceneOutlinerTreeItemPtr> ComponentRows;
-	if (bShowingComponents)
+	bool bAnyItemsBeforeFiltering = false;
+
+	check(Actor);
+	if (const UObjectMixerObjectFilter* Filter = GetObjectFilter())
 	{
-		TArray<UActorComponent*> ActorComponents;
-		Actor->GetComponents(ActorComponents);
-		for (UActorComponent* Component : ActorComponents)
+		const TArray<AActor*> AssociatedActors = Filter->FindAssociatedActors(Actor);
+
+		for (AActor* AssociatedActor : AssociatedActors)
 		{
-			if (IsValid(Component) && DoesWorldObjectHaveAcceptableClass(Component))
+			if (IsValid(Actor))
 			{
-				bWouldReturnAnyComponentsBeforeFiltering = true;
-					
-				if (FSceneOutlinerTreeItemPtr ComponentItem =
-					Mode->CreateItemFor<FObjectMixerEditorListRowComponent>(
-						FObjectMixerEditorListRowComponent(Component, GetCastedMode()->GetSceneOutliner())))
+				bAnyItemsBeforeFiltering = true;
+
+				if (FSceneOutlinerTreeItemPtr ObjectItem =
+					Mode->CreateItemFor<FObjectMixerEditorListRowActor>(
+						FObjectMixerEditorListRowActor(AssociatedActor, GetCastedMode()->GetSceneOutliner(), FText::GetEmpty(), Actor)))
 				{
-					ComponentRows.Add(ComponentItem);
+					OutItems.Add(ObjectItem);
 				}
 			}
 		}
 	}
 
+	return bAnyItemsBeforeFiltering;
+}
+
+TArray<FSceneOutlinerTreeItemPtr> FObjectMixerOutlinerHierarchy::ConditionallyCreateActorAndComponentItems(AActor* Actor) const
+{
+	if (!IsValid(Actor))
+	{
+		return {};
+	}
+	
+	// Whether or not we have children to return, we should create an actor row if children would be returned if there were no filters
+	bool bWouldReturnAnyChildItemsBeforeFiltering = false;
+	
+	TArray<FSceneOutlinerTreeItemPtr> ComponentRows;
+	bWouldReturnAnyChildItemsBeforeFiltering |= CreateComponentItems(Actor, ComponentRows);
+
+	TArray<FSceneOutlinerTreeItemPtr> AssociatedActorRows;
+	bWouldReturnAnyChildItemsBeforeFiltering |= CreateAssociatedActorItems(Actor, AssociatedActorRows);
+
 	const bool bShouldCreateActorItem =
-		bWouldReturnAnyComponentsBeforeFiltering || !bShowingOnlyActorWithValidComponents || DoesWorldObjectHaveAcceptableClass(Actor);
+		bWouldReturnAnyChildItemsBeforeFiltering || !bShowingOnlyActorWithValidComponents || DoesWorldObjectHaveAcceptableClass(Actor);
 
 	TArray<FSceneOutlinerTreeItemPtr> ReturnValue;
 	if (bShouldCreateActorItem)
@@ -608,7 +652,7 @@ TArray<FSceneOutlinerTreeItemPtr> FObjectMixerOutlinerHierarchy::ConditionallyCr
 				FObjectMixerEditorListRowActor(Actor, GetCastedMode()->GetSceneOutliner())))
 		{
 			const bool bIsHybridModeEnabled = GetCastedMode()->ShouldAllowHybridRows() && GetDefault<UObjectMixerEditorSettings>()->IsHybridRowModeEnabled();
-			if (bIsHybridModeEnabled && ComponentRows.Num() == 1) // Create hybrid row
+			if (bIsHybridModeEnabled && ComponentRows.Num() == 1 && AssociatedActorRows.IsEmpty()) // Create hybrid row
 			{
 				if (FObjectMixerEditorListRowActor* AsActorRow = FObjectMixerUtils::AsActorRow(ActorItem))
 				{
@@ -623,6 +667,7 @@ TArray<FSceneOutlinerTreeItemPtr> FObjectMixerOutlinerHierarchy::ConditionallyCr
 			else
 			{
 				ReturnValue.Append(ComponentRows);
+				ReturnValue.Append(AssociatedActorRows);
 			}
 
 			// Place ActorItem before components
@@ -680,6 +725,7 @@ void FObjectMixerOutlinerHierarchy::CreateChildren(const FSceneOutlinerTreeItemP
 		}
 
 		CreateComponentItems(ParentActor, OutChildren);
+		CreateAssociatedActorItems(ParentActor, OutChildren);
 
 		TArray<AActor*> ChildActors;
 
diff --git a/Engine/Plugins/Editor/ObjectMixer/ObjectMixer/Source/ObjectMixer/Private/Views/List/Modes/ObjectMixerOutlinerMode.cpp b/Engine/Plugins/Editor/ObjectMixer/ObjectMixer/Source/ObjectMixer/Private/Views/List/Modes/ObjectMixerOutlinerMode.cpp
--- a/Engine/Plugins/Editor/ObjectMixer/ObjectMixer/Source/ObjectMixer/Private/Views/List/Modes/ObjectMixerOutlinerMode.cpp
+++ b/Engine/Plugins/Editor/ObjectMixer/ObjectMixer/Source/ObjectMixer/Private/Views/List/Modes/ObjectMixerOutlinerMode.cpp
@@ -903,6 +903,22 @@ TArray<UActorComponent*> FObjectMixerOutlinerMode::GetSelectedComponentsInEditor
 	return TArray<UActorComponent*>();
 }
 
+bool FObjectMixerOutlinerMode::GetIsObjectOfFilteredClass(const UObject* Object) const
+{
+	if (IsValid(Object))
+	{
+		for (UClass* Class : GetListModelPtr().Pin()->ObjectClassesToFilterCache)
+		{
+			if (Object->IsA(Class))
+			{
+				return true;
+			}
+		}
+	}
+
+	return false;
+}
+
 void FObjectMixerOutlinerMode::OnActorLabelChanged(AActor* ChangedActor)
 {
 	if (!ensure(ChangedActor))
@@ -910,7 +926,13 @@ void FObjectMixerOutlinerMode::OnActorLabelChanged(AActor* ChangedActor)
 		return;
 	}
 
-	if (IsActorDisplayable(ChangedActor) && RepresentingWorld.Get() == ChangedActor->GetWorld())
+	if (!IsActorDisplayable(ChangedActor) || RepresentingWorld.Get() != ChangedActor->GetWorld())
+	{
+		return;
+	}
+
+	// Check if this actor is one of the registered types
+	if (GetIsObjectOfFilteredClass(ChangedActor))
 	{
 		// Force create the item otherwise the outliner may not be notified of a change to the item if it is filtered out
 		if (FSceneOutlinerTreeItemPtr Item = CreateItemFor<FObjectMixerEditorListRowActor>(
@@ -919,6 +941,27 @@ void FObjectMixerOutlinerMode::OnActorLabelChanged(AActor* ChangedActor)
 			SceneOutliner->OnItemLabelChanged(Item);
 		}
 	}
+	else if (const UObjectMixerObjectFilter* Filter = GetListModelPtr().Pin()->GetMainObjectFilterInstance())
+	{
+		// This actor is likely an associated actor of one or more registered type actors, so we need to force an update for each parent/associated pair
+		const TSet<TWeakObjectPtr<AActor>> ApplicableActorsCopy = ApplicableActors; // This set may change while we're iterating, so make a copy
+		for (TWeakObjectPtr<AActor> ApplicableActor : ApplicableActorsCopy)
+		{
+			if (!ApplicableActor.IsValid())
+			{
+				continue;
+			}
+
+			if (Filter->IsActorAssociated(ApplicableActor.Get(), ChangedActor))
+			{
+				if (FSceneOutlinerTreeItemPtr Item = CreateItemFor<FObjectMixerEditorListRowActor>(
+					FObjectMixerEditorListRowActor(ChangedActor, GetSceneOutliner(), FText::GetEmpty(), ApplicableActor.Get()), true))
+				{
+					SceneOutliner->OnItemLabelChanged(Item);
+				}
+			}
+		}
+	}
 }
 
 void FObjectMixerOutlinerMode::OnObjectsReplaced(const TMap<UObject*, UObject*>& ReplacementMap)
@@ -1889,6 +1932,20 @@ void FObjectMixerOutlinerMode::OnItemAdded(FSceneOutlinerTreeItemPtr Item)
 			// Synchronize selection with editor if enabled
 			if (ShouldSyncSelectionFromEditor() && GetSelectedActorsInEditor().Contains(ActorItem->Actor.Get()))
 			{
+				if (const FObjectMixerEditorListRowActor* ObjectMixerItem = Item->CastTo<FObjectMixerEditorListRowActor>())
+				{
+					if (ObjectMixerItem->OverrideParent)
+					{
+						// If there's an OverrideParent, this actor may appear multiple times. We disambiguate the selection
+						// based on a GUID that combines the parent and actor itself, but the general editor doesn't. To avoid
+						// selecting the wrong one, skip selecting this new item if it wasn't previously selected.
+						if (!SelectedIDs.Contains(ObjectMixerItem->GetID()))
+						{
+							return;
+						}
+					}
+				}
+
 				SceneOutliner->SetItemSelection(Item, true);
 			}
 		}
@@ -2466,6 +2523,14 @@ void FObjectMixerOutlinerMode::SynchronizeSelectedActorDescs()
 
 FSceneOutlinerDragValidationInfo FObjectMixerOutlinerMode::ValidateDrop(const ISceneOutlinerTreeItem& DropTarget, const FSceneOutlinerDragDropPayload& Payload) const
 {
+	if (const UObjectMixerObjectFilter* Filter = GetListModelPtr().Pin()->GetMainObjectFilterInstance())
+	{
+		if (Filter->HasCustomDropHandling(DropTarget))
+		{
+			return Filter->ValidateDrop(DropTarget, Payload);
+		}
+	}
+
 	if (Payload.Has<FFolderTreeItem>())
 	{
 		FFolder::FRootObject TargetRootObject = DropTarget.GetRootObject();
@@ -2759,6 +2824,15 @@ FSceneOutlinerDragValidationInfo FObjectMixerOutlinerMode::ValidateDrop(const IS
 
 void FObjectMixerOutlinerMode::OnDrop(ISceneOutlinerTreeItem& DropTarget, const FSceneOutlinerDragDropPayload& Payload, const FSceneOutlinerDragValidationInfo& ValidationInfo) const
 {
+	if (const UObjectMixerObjectFilter* Filter = GetListModelPtr().Pin()->GetMainObjectFilterInstance())
+	{
+		if (Filter->HasCustomDropHandling(DropTarget))
+		{
+			Filter->OnDrop(DropTarget, Payload, ValidationInfo);
+			return;
+		}
+	}
+
 	if (const FActorTreeItem* ActorItem = DropTarget.CastTo<FActorTreeItem>())
 	{
 		AActor* DropActor = ActorItem->Actor.Get();
diff --git a/Engine/Plugins/Editor/ObjectMixer/ObjectMixer/Source/ObjectMixer/Private/Views/List/RowTypes/ObjectMixerEditorListRowActor.cpp b/Engine/Plugins/Editor/ObjectMixer/ObjectMixer/Source/ObjectMixer/Private/Views/List/RowTypes/ObjectMixerEditorListRowActor.cpp
--- a/Engine/Plugins/Editor/ObjectMixer/ObjectMixer/Source/ObjectMixer/Private/Views/List/RowTypes/ObjectMixerEditorListRowActor.cpp
+++ b/Engine/Plugins/Editor/ObjectMixer/ObjectMixer/Source/ObjectMixer/Private/Views/List/RowTypes/ObjectMixerEditorListRowActor.cpp
@@ -18,3 +18,14 @@ void FObjectMixerEditorListRowActor::OnVisibilityChanged(const bool bNewVisibili
 		ListView->EvaluateAndSetEditorVisibilityPerRow();
 	}
 }
+
+FSceneOutlinerTreeItemID FObjectMixerEditorListRowActor::GetID() const
+{
+	if (!OverrideParent.IsValid())
+	{
+		return FActorTreeItem::GetID();
+	}
+
+	// The same actor can appear as the child of multiple override parents, so generate unique IDs to let it display separately under each parent
+	return FSceneOutlinerTreeItemID(FGuid::Combine(GetGuid(), OverrideParent.Get()->GetActorGuid()));
+}
diff --git a/Engine/Plugins/Editor/ObjectMixer/ObjectMixer/Source/ObjectMixer/Public/ObjectFilter/ObjectMixerEditorObjectFilter.h b/Engine/Plugins/Editor/ObjectMixer/ObjectMixer/Source/ObjectMixer/Public/ObjectFilter/ObjectMixerEditorObjectFilter.h
--- a/Engine/Plugins/Editor/ObjectMixer/ObjectMixer/Source/ObjectMixer/Public/ObjectFilter/ObjectMixerEditorObjectFilter.h
+++ b/Engine/Plugins/Editor/ObjectMixer/ObjectMixer/Source/ObjectMixer/Public/ObjectFilter/ObjectMixerEditorObjectFilter.h
@@ -3,6 +3,7 @@
 #pragma once
 
 #include "GameFramework/Actor.h"
+#include "SceneOutlinerDragDrop.h"
 #include "UObject/Class.h"
 
 #include "ObjectMixerEditorObjectFilter.generated.h"
@@ -137,6 +138,29 @@ public:
 	 */
 	virtual TSet<FName> GetPropertiesThatRequireListRefresh() const;
 
+	/**
+	 * Fills OutAssociatedActors with additional actors to include as children in the ObjectMixer hierarchy for the given actor.
+	 */
+	virtual TArray<AActor*> FindAssociatedActors(AActor* InActor) const;
+
+	/**
+	 * Returns true if AssociatedActor is associated with Actor.
+	 * This should return true if the actor would be included in FindAssociatedActors, but is a separate function to avoid allocating an array every time.
+	 */
+	virtual bool IsActorAssociated(AActor* Actor, AActor* AssociatedActor) const;
+
+	/** 
+	 * Returns true if the given tree item should override the Object Mixer's standard handling for drop operations (ValidateDrop/OnDrop)
+	 */
+	virtual bool HasCustomDropHandling(const ISceneOutlinerTreeItem& DropTarget) const;
+
+	/** Test whether the specified payload can be dropped onto a tree item */
+	virtual FSceneOutlinerDragValidationInfo ValidateDrop(const ISceneOutlinerTreeItem& DropTarget, const FSceneOutlinerDragDropPayload& Payload) const;
+
+	/** Called when a payload is dropped onto a target */
+	virtual void OnDrop(ISceneOutlinerTreeItem& DropTarget, const FSceneOutlinerDragDropPayload& Payload, const FSceneOutlinerDragValidationInfo& ValidationInfo) const;
+
+
 	static TSet<UClass*> GetParentAndChildClassesFromSpecifiedClasses(
 		const TSet<UClass*>& InSpecifiedClasses, EObjectMixerInheritanceInclusionOptions Options);
 
diff --git a/Engine/Plugins/Editor/ObjectMixer/ObjectMixer/Source/ObjectMixer/Public/Views/List/Modes/ObjectMixerOutlinerHierarchy.h b/Engine/Plugins/Editor/ObjectMixer/ObjectMixer/Source/ObjectMixer/Public/Views/List/Modes/ObjectMixerOutlinerHierarchy.h
--- a/Engine/Plugins/Editor/ObjectMixer/ObjectMixer/Source/ObjectMixer/Public/Views/List/Modes/ObjectMixerOutlinerHierarchy.h
+++ b/Engine/Plugins/Editor/ObjectMixer/ObjectMixer/Source/ObjectMixer/Public/Views/List/Modes/ObjectMixerOutlinerHierarchy.h
@@ -8,6 +8,7 @@ class FObjectMixerOutlinerMode;
 class FWorldPartitionActorDescInstance;
 class UActorFolder;
 class ULevelInstanceSubsystem;
+class UObjectMixerObjectFilter;
 class UWorldPartition;
 
 class FObjectMixerOutlinerHierarchy : public ISceneOutlinerHierarchy, public FNoncopyable
@@ -18,6 +19,7 @@ public:
 	static TUniquePtr<FObjectMixerOutlinerHierarchy> Create(class ISceneOutlinerMode* Mode, const TWeakObjectPtr<UWorld>& World);
 
 	FObjectMixerOutlinerMode* GetCastedMode() const;
+	const UObjectMixerObjectFilter* GetObjectFilter() const;
 
 	/** Create a linearization of all applicable items in the hierarchy */
 	virtual void CreateItems(TArray<FSceneOutlinerTreeItemPtr>& OutItems) const override;
@@ -45,8 +47,15 @@ protected:
 	
 	/**
 	 * Create all component items for an actor if we are showing components and place them in OutItems. Items not added if filtered out.
+	 * @return true if any items would have been returned before filtering
 	 */
-	void CreateComponentItems(const AActor* Actor, TArray<FSceneOutlinerTreeItemPtr>& OutItems) const;
+	bool CreateComponentItems(AActor* Actor, TArray<FSceneOutlinerTreeItemPtr>& OutItems) const;
+
+	/**
+	 * Create all associated actor items for an actor and place them in OutItems. Items not added if filtered out.
+	 * @return true if any items would have been returned before filtering
+	 */
+	bool CreateAssociatedActorItems(AActor* Actor, TArray<FSceneOutlinerTreeItemPtr>& OutItems) const;
 
 	/**
 	 * Create all component items for an actor if we are showing components and place them in OutItems. Items not added if filtered out.
diff --git a/Engine/Plugins/Editor/ObjectMixer/ObjectMixer/Source/ObjectMixer/Public/Views/List/Modes/ObjectMixerOutlinerMode.h b/Engine/Plugins/Editor/ObjectMixer/ObjectMixer/Source/ObjectMixer/Public/Views/List/Modes/ObjectMixerOutlinerMode.h
--- a/Engine/Plugins/Editor/ObjectMixer/ObjectMixer/Source/ObjectMixer/Public/Views/List/Modes/ObjectMixerOutlinerMode.h
+++ b/Engine/Plugins/Editor/ObjectMixer/ObjectMixer/Source/ObjectMixer/Public/Views/List/Modes/ObjectMixerOutlinerMode.h
@@ -318,6 +318,9 @@ protected:
 
 	/** Get the list of components selected in the editor */
 	TArray<UActorComponent*> GetSelectedComponentsInEditor() const;
+
+	/** Check if an object is a member of one of the classes our list model is filtered to */
+	bool GetIsObjectOfFilteredClass(const UObject* Object) const;
 	
 	/** Build and up the context menu */
 	TSharedPtr<SWidget> BuildContextMenu();
diff --git a/Engine/Plugins/Editor/ObjectMixer/ObjectMixer/Source/ObjectMixer/Public/Views/List/RowTypes/ObjectMixerEditorListRowActor.h b/Engine/Plugins/Editor/ObjectMixer/ObjectMixer/Source/ObjectMixer/Public/Views/List/RowTypes/ObjectMixerEditorListRowActor.h
--- a/Engine/Plugins/Editor/ObjectMixer/ObjectMixer/Source/ObjectMixer/Public/Views/List/RowTypes/ObjectMixerEditorListRowActor.h
+++ b/Engine/Plugins/Editor/ObjectMixer/ObjectMixer/Source/ObjectMixer/Public/Views/List/RowTypes/ObjectMixerEditorListRowActor.h
@@ -12,9 +12,12 @@ struct OBJECTMIXEREDITOR_API FObjectMixerEditorListRowActor : FActorTreeItem
 {
 	explicit FObjectMixerEditorListRowActor(
 		AActor* InObject, 
-		SSceneOutliner* InSceneOutliner, const FText& InDisplayNameOverride = FText::GetEmpty())
+		SSceneOutliner* InSceneOutliner,
+		const FText& InDisplayNameOverride = FText::GetEmpty(),
+		AActor* InOverrideParent = nullptr)
 	: FActorTreeItem(InObject)
 	, OriginalObjectSoftPtr(InObject)
+	, OverrideParent(InOverrideParent)
 	{
 		TreeType = Type;
 		RowData = FObjectMixerEditorListRowData(InSceneOutliner, InDisplayNameOverride);
@@ -25,8 +28,12 @@ struct OBJECTMIXEREDITOR_API FObjectMixerEditorListRowActor : FActorTreeItem
 	/** Used in scenarios where the original object may be reconstructed or trashed, such as when running a construction script. */
 	TSoftObjectPtr<AActor> OriginalObjectSoftPtr;
 
+	/** If set, the row for this actor will be treated as the current row's parent in the hierarchy */
+	TSoftObjectPtr<AActor> OverrideParent;
+
 	/* Begin ISceneOutlinerTreeItem Implementation */
 	static const FSceneOutlinerTreeItemType Type;
 	virtual void OnVisibilityChanged(const bool bNewVisibility) override;
+	virtual FSceneOutlinerTreeItemID GetID() const override;
 	/* End ISceneOutlinerTreeItem Implementation */
 };
diff --git a/Engine/Plugins/Experimental/ColorCorrectRegions/Source/ColorCorrectRegionsEditor/Private/ColorCorrectRegionsEditorModule.cpp b/Engine/Plugins/Experimental/ColorCorrectRegions/Source/ColorCorrectRegionsEditor/Private/ColorCorrectRegionsEditorModule.cpp
--- a/Engine/Plugins/Experimental/ColorCorrectRegions/Source/ColorCorrectRegionsEditor/Private/ColorCorrectRegionsEditorModule.cpp
+++ b/Engine/Plugins/Experimental/ColorCorrectRegions/Source/ColorCorrectRegionsEditor/Private/ColorCorrectRegionsEditorModule.cpp
@@ -9,6 +9,7 @@
 #include "ColorCorrectWindow.h"
 #include "ColorGradingDataModelGenerator_ColorCorrectRegion.h"
 #include "ColorGradingEditorDataModel.h"
+#include "ColorGradingHierarchyConfig_ColorCorrectRegion.h"
 #include "ColorGradingMixerObjectFilterRegistry.h"
 #include "IPlacementModeModule.h"
 #include "PropertyEditorModule.h"
@@ -29,7 +30,10 @@ void FColorCorrectRegionsEditorModule::StartupModule()
 	FColorGradingMixerObjectFilterRegistry::RegisterActorClassToPlace(AColorCorrectionRegion::StaticClass());
 	FColorGradingMixerObjectFilterRegistry::RegisterActorClassToPlace(AColorCorrectionWindow::StaticClass());
 
-	FColorGradingMixerObjectFilterRegistry::RegisterObjectClassToFilter(AColorCorrectRegion::StaticClass());
+	FColorGradingMixerObjectFilterRegistry::RegisterObjectClassToFilter(
+		AColorCorrectRegion::StaticClass(),
+		FGetObjectHierarchyConfig::CreateStatic(&FColorGradingHierarchyConfig_ColorCorrectRegion::MakeInstance)
+	);
 
 	RegisterSectionMappings();
 
diff --git a/Engine/Plugins/Experimental/ColorCorrectRegions/Source/ColorCorrectRegionsEditor/Private/ColorGradingDataModelGenerator_ColorCorrectRegion.cpp b/Engine/Plugins/Experimental/ColorCorrectRegions/Source/ColorCorrectRegionsEditor/Private/ColorGradingDataModelGenerator_ColorCorrectRegion.cpp
--- a/Engine/Plugins/Experimental/ColorCorrectRegions/Source/ColorCorrectRegionsEditor/Private/ColorGradingDataModelGenerator_ColorCorrectRegion.cpp
+++ b/Engine/Plugins/Experimental/ColorCorrectRegions/Source/ColorCorrectRegionsEditor/Private/ColorGradingDataModelGenerator_ColorCorrectRegion.cpp
@@ -22,7 +22,8 @@
 
 const TArray<FName> VisibleCategories = {
 	TEXT("Region"),
-	TEXT("Color Grading")
+	TEXT("Color Grading"),
+	TEXT("PerActorColorCorrection")
 };
 
 TSharedRef<IColorGradingEditorDataModelGenerator> FColorGradingDataModelGenerator_ColorCorrectRegion::MakeInstance()
@@ -63,6 +64,11 @@ public:
 		// TransformCommon is a custom category that doesn't get returned by GetCategoryNames that also needs to be hidden
 		DetailBuilder.HideCategory(TEXT("TransformCommon"));
 
+		IDetailCategoryBuilder& PerActorCCCategoryBuilder = DetailBuilder.EditCategory(TEXT("PerActorColorCorrection"), LOCTEXT("PerActorCCDisplayName", "Per-Actor Color Correction"));
+
+		PerActorCCCategoryBuilder.AddProperty(DetailBuilder.GetProperty((GET_MEMBER_NAME_CHECKED(AColorCorrectionRegion, bEnablePerActorCC))));
+		PerActorCCCategoryBuilder.AddProperty(DetailBuilder.GetProperty((GET_MEMBER_NAME_CHECKED(AColorCorrectionRegion, PerActorColorCorrection))));
+
 		// Hide CCR-specific properties if CCWs are present in the selection
 		const bool bHasCCWs = DetailBuilder.GetSelectedObjects().ContainsByPredicate([](const TWeakObjectPtr<UObject>& SelectedObject)
 		{
diff --git a/Engine/Plugins/Experimental/ColorCorrectRegions/Source/ColorCorrectRegionsEditor/Private/ColorGradingHierarchyConfig_ColorCorrectRegion.cpp b/Engine/Plugins/Experimental/ColorCorrectRegions/Source/ColorCorrectRegionsEditor/Private/ColorGradingHierarchyConfig_ColorCorrectRegion.cpp
new file mode 100644
--- /dev/null
+++ b/Engine/Plugins/Experimental/ColorCorrectRegions/Source/ColorCorrectRegionsEditor/Private/ColorGradingHierarchyConfig_ColorCorrectRegion.cpp
@@ -0,0 +1,125 @@
+// Copyright Epic Games, Inc. All Rights Reserved.
+
+#include "ColorGradingHierarchyConfig_ColorCorrectRegion.h"
+
+#include "ActorTreeItem.h"
+#include "ColorCorrectRegion.h"
+#include "DragAndDrop/ActorDragDropOp.h"
+#include "Editor.h"
+
+#define LOCTEXT_NAMESPACE "ColorGradingHierarchyConfig_ColorCorrectRegion"
+
+TSharedRef<IColorGradingMixerObjectHierarchyConfig> FColorGradingHierarchyConfig_ColorCorrectRegion::MakeInstance()
+{
+	return MakeShared<FColorGradingHierarchyConfig_ColorCorrectRegion>();
+}
+
+TArray<AActor*> FColorGradingHierarchyConfig_ColorCorrectRegion::FindAssociatedActors(UObject* ParentObject) const
+{
+	TArray<AActor*> AssociatedActors;
+
+	// List affected per-actor CC targets as associated actors so they appear as children in the hierarchy
+	if (const AColorCorrectRegion* Region = Cast<AColorCorrectRegion>(ParentObject))
+	{
+		if (Region->bEnablePerActorCC)
+		{
+			for (TSoftObjectPtr<AActor> Actor : Region->AffectedActors)
+			{
+				if (Actor.IsValid())
+				{
+					AssociatedActors.Add(Actor.Get());
+				}
+			}
+		}
+	}
+
+	return MoveTemp(AssociatedActors);
+}
+
+bool FColorGradingHierarchyConfig_ColorCorrectRegion::IsActorAssociated(UObject* ParentObject, AActor* AssociatedActor) const
+{
+	if (const AColorCorrectRegion* Region = Cast<AColorCorrectRegion>(ParentObject))
+	{
+		if (Region->bEnablePerActorCC)
+		{
+			return Region->AffectedActors.Contains(AssociatedActor);
+		}
+	}
+
+	return false;
+}
+
+bool FColorGradingHierarchyConfig_ColorCorrectRegion::HasCustomDropHandling() const
+{
+	return true;
+}
+
+FSceneOutlinerDragValidationInfo FColorGradingHierarchyConfig_ColorCorrectRegion::ValidateDrop(UObject* DropTarget, const FSceneOutlinerDragDropPayload& Payload) const
+{
+	if (AColorCorrectRegion* Region = Cast<AColorCorrectRegion>(DropTarget))
+	{
+		if (Payload.SourceOperation.IsOfType<FActorDragDropOp>())
+		{
+			return FSceneOutlinerDragValidationInfo(
+				ESceneOutlinerDropCompatibility::CompatibleAttach,
+				FText::Format(
+					LOCTEXT("FoldersOnActorError", "Add to Per-Actor CC for {0}"),
+					FText::FromString(Region->GetActorLabel())
+				)
+			);
+		}
+	}
+
+	return FSceneOutlinerDragValidationInfo::Invalid();
+}
+
+bool FColorGradingHierarchyConfig_ColorCorrectRegion::OnDrop(UObject* DropTarget, const FSceneOutlinerDragDropPayload& Payload, const FSceneOutlinerDragValidationInfo& ValidationInfo) const
+{
+	if (AColorCorrectRegion* Region = Cast<AColorCorrectRegion>(DropTarget))
+	{
+		if (Payload.SourceOperation.IsOfType<FActorDragDropOp>())
+		{
+			GEditor->BeginTransaction(LOCTEXT("AddToPerActorCCTransaction", "Add actors to Per-Actor CC"));
+
+			// Enable per-actor CC if not already enabled
+			if (!Region->bEnablePerActorCC)
+			{
+				const FName EnablePerActorCCPropertyName = GET_MEMBER_NAME_CHECKED(AColorCorrectRegion, bEnablePerActorCC);
+				FProperty* EnablePerActorCCProperty = FindFieldChecked<FProperty>(AColorCorrectRegion::StaticClass(), EnablePerActorCCPropertyName);
+				Region->PreEditChange(EnablePerActorCCProperty);
+
+				Region->bEnablePerActorCC = true;
+
+				FPropertyChangedEvent PropertyEvent(EnablePerActorCCProperty);
+				PropertyEvent.ChangeType = EPropertyChangeType::ValueSet;
+				Region->PostEditChangeProperty(PropertyEvent);
+			}
+
+			// Add affected actors from dragged list
+			const FName AffectedActorsPropertyName = GET_MEMBER_NAME_CHECKED(AColorCorrectRegion, AffectedActors);
+			FProperty* AffectedActorsProperty = FindFieldChecked<FProperty>(AColorCorrectRegion::StaticClass(), AffectedActorsPropertyName);
+			Region->PreEditChange(AffectedActorsProperty);
+
+			const FActorDragDropOp& ActorDragDropOp = StaticCast<const FActorDragDropOp&>(Payload.SourceOperation);
+			for (TWeakObjectPtr<AActor> Actor : ActorDragDropOp.Actors)
+			{
+				if (Actor.IsValid())
+				{
+					Region->AffectedActors.Add(Actor.Get());
+				}
+			}
+
+			FPropertyChangedEvent PropertyEvent(AffectedActorsProperty);
+			PropertyEvent.ChangeType = EPropertyChangeType::ArrayAdd;
+			Region->PostEditChangeProperty(PropertyEvent);
+
+			GEditor->EndTransaction();
+
+			return true;
+		}
+	}
+
+	return false;
+}
+
+#undef LOCTEXT_NAMESPACE
\ No newline at end of file
diff --git a/Engine/Plugins/Experimental/ColorCorrectRegions/Source/ColorCorrectRegionsEditor/Private/ColorGradingHierarchyConfig_ColorCorrectRegion.h b/Engine/Plugins/Experimental/ColorCorrectRegions/Source/ColorCorrectRegionsEditor/Private/ColorGradingHierarchyConfig_ColorCorrectRegion.h
new file mode 100644
--- /dev/null
+++ b/Engine/Plugins/Experimental/ColorCorrectRegions/Source/ColorCorrectRegionsEditor/Private/ColorGradingHierarchyConfig_ColorCorrectRegion.h
@@ -0,0 +1,21 @@
+// Copyright Epic Games, Inc. All Rights Reserved.
+
+#pragma once
+
+#include "ColorGradingMixerObjectFilterRegistry.h"
+
+/** Color Grading hierarchy config for the AColorCorrectionRegion actor class */
+struct FColorGradingHierarchyConfig_ColorCorrectRegion : IColorGradingMixerObjectHierarchyConfig
+{
+public:
+	static TSharedRef<IColorGradingMixerObjectHierarchyConfig> MakeInstance();
+
+	//~ IColorGradingMixerObjectHierarchyConfig interface
+	virtual TArray<AActor*> FindAssociatedActors(UObject* ParentObject) const override;
+	virtual bool IsActorAssociated(UObject* ParentObject, AActor* AssociatedActor) const override;
+	virtual bool HasCustomDropHandling() const override;
+	virtual FSceneOutlinerDragValidationInfo ValidateDrop(UObject* DropTarget, const FSceneOutlinerDragDropPayload& Payload) const override;
+	virtual bool OnDrop(UObject* DropTarget, const FSceneOutlinerDragDropPayload& Payload, const FSceneOutlinerDragValidationInfo& ValidationInfo) const override;
+
+	//~ End IColorGradingMixerObjectHierarchyConfig interface
+};
\ No newline at end of file
