Updated headers using UnrealCodeFixup to make sure dllstorage is on methods/staticvars instead of types. Part 2/n

diff --git a/Engine/Plugins/Editor/ObjectMixer/ObjectMixer/Source/ObjectMixer/Public/Views/List/Modes/ObjectMixerOutlinerMode.h b/Engine/Plugins/Editor/ObjectMixer/ObjectMixer/Source/ObjectMixer/Public/Views/List/Modes/ObjectMixerOutlinerMode.h
--- a/Engine/Plugins/Editor/ObjectMixer/ObjectMixer/Source/ObjectMixer/Public/Views/List/Modes/ObjectMixerOutlinerMode.h
+++ b/Engine/Plugins/Editor/ObjectMixer/ObjectMixer/Source/ObjectMixer/Public/Views/List/Modes/ObjectMixerOutlinerMode.h
@@ -13,6 +13,8 @@
 
 #include "ObjectMixerOutlinerMode.generated.h"
 
+#define UE_API OBJECTMIXEREDITOR_API
+
 class FObjectMixerEditorList;
 class FObjectMixerOutlinerMode;
 class IWorldPartitionEditorModule;
@@ -138,7 +140,7 @@ private:
 	static TObjectPtr<UObjectMixerOutlinerModeEditorConfig> Instance;
 };
 
-struct OBJECTMIXEREDITOR_API FObjectMixerOutlinerModeParams : FActorModeParams
+struct FObjectMixerOutlinerModeParams : FActorModeParams
 {
 	FObjectMixerOutlinerModeParams() {}
 
@@ -153,7 +155,7 @@ struct OBJECTMIXEREDITOR_API FObjectMixerOutlinerModeParams : FActorModeParams
 	}
 };
 
-class OBJECTMIXEREDITOR_API FObjectMixerOutlinerMode : public FActorMode
+class FObjectMixerOutlinerMode : public FActorMode
 {
 public:
 	
@@ -163,8 +165,8 @@ public:
 		bool bIsUserSelected;
 	};
 	
-	FObjectMixerOutlinerMode(const FObjectMixerOutlinerModeParams& Params, const TSharedRef<FObjectMixerEditorList> InListModel);
-	virtual ~FObjectMixerOutlinerMode();
+	UE_API FObjectMixerOutlinerMode(const FObjectMixerOutlinerModeParams& Params, const TSharedRef<FObjectMixerEditorList> InListModel);
+	UE_API virtual ~FObjectMixerOutlinerMode();
 
 	TWeakPtr<FObjectMixerEditorList> GetListModelPtr() const
 	{
@@ -179,36 +181,36 @@ public:
 	/**
 	 * Determines the style of the tree (flat list or hierarchy)
 	 */
-	EObjectMixerTreeViewMode GetTreeViewMode() const;
+	UE_API EObjectMixerTreeViewMode GetTreeViewMode() const;
 	/**
 	 * Determine the style of the tree (flat list or hierarchy)
 	 */
-	void SetTreeViewMode(EObjectMixerTreeViewMode InViewMode);
+	UE_API void SetTreeViewMode(EObjectMixerTreeViewMode InViewMode);
 
-	UWorld* GetRepresentingWorld();
+	UE_API UWorld* GetRepresentingWorld();
 
 	/* Begin ISceneOutlinerMode Interface */
 	virtual bool IsInteractive() const override { return true; }
-	virtual TUniquePtr<ISceneOutlinerHierarchy> CreateHierarchy() override;
-	virtual void Rebuild() override;
-	virtual FCreateSceneOutlinerMode CreateFolderPickerMode(const FFolder::FRootObject& InRootObject = FFolder::GetInvalidRootObject()) const override;
-	virtual void CreateViewContent(FMenuBuilder& MenuBuilder) override;
-	virtual TSharedPtr<FDragDropOperation> CreateDragDropOperation(const FPointerEvent& MouseEvent, const TArray<FSceneOutlinerTreeItemPtr>& InTreeItems) const override;
-	virtual bool ParseDragDrop(FSceneOutlinerDragDropPayload& OutPayload, const FDragDropOperation& Operation) const override;
-	virtual FSceneOutlinerDragValidationInfo ValidateDrop(const ISceneOutlinerTreeItem& DropTarget, const FSceneOutlinerDragDropPayload& Payload) const override;
-	virtual TSharedPtr<SWidget> CreateContextMenu() override;
-	virtual void OnItemAdded(FSceneOutlinerTreeItemPtr Item) override;
-	virtual void OnItemRemoved(FSceneOutlinerTreeItemPtr Item) override;
-	virtual void OnItemSelectionChanged(FSceneOutlinerTreeItemPtr Item, ESelectInfo::Type SelectionType, const FSceneOutlinerItemSelection& Selection) override;
-	virtual void OnItemDoubleClick(FSceneOutlinerTreeItemPtr Item) override;
-	virtual void OnFilterTextCommited(FSceneOutlinerItemSelection& Selection, ETextCommit::Type CommitType) override;
-	virtual void OnItemPassesFilters(const ISceneOutlinerTreeItem& Item) override;
-	virtual FReply OnKeyDown(const FKeyEvent& InKeyEvent) override;
-	virtual void OnDuplicateSelected() override;
-	virtual void OnDrop(ISceneOutlinerTreeItem& DropTarget, const FSceneOutlinerDragDropPayload& Payload, const FSceneOutlinerDragValidationInfo& ValidationInfo) const override;
-	virtual bool CanRenameItem(const ISceneOutlinerTreeItem& Item) const override;
-	virtual FText GetStatusText() const override;
-	virtual FSlateColor GetStatusTextColor() const override;
+	UE_API virtual TUniquePtr<ISceneOutlinerHierarchy> CreateHierarchy() override;
+	UE_API virtual void Rebuild() override;
+	UE_API virtual FCreateSceneOutlinerMode CreateFolderPickerMode(const FFolder::FRootObject& InRootObject = FFolder::GetInvalidRootObject()) const override;
+	UE_API virtual void CreateViewContent(FMenuBuilder& MenuBuilder) override;
+	UE_API virtual TSharedPtr<FDragDropOperation> CreateDragDropOperation(const FPointerEvent& MouseEvent, const TArray<FSceneOutlinerTreeItemPtr>& InTreeItems) const override;
+	UE_API virtual bool ParseDragDrop(FSceneOutlinerDragDropPayload& OutPayload, const FDragDropOperation& Operation) const override;
+	UE_API virtual FSceneOutlinerDragValidationInfo ValidateDrop(const ISceneOutlinerTreeItem& DropTarget, const FSceneOutlinerDragDropPayload& Payload) const override;
+	UE_API virtual TSharedPtr<SWidget> CreateContextMenu() override;
+	UE_API virtual void OnItemAdded(FSceneOutlinerTreeItemPtr Item) override;
+	UE_API virtual void OnItemRemoved(FSceneOutlinerTreeItemPtr Item) override;
+	UE_API virtual void OnItemSelectionChanged(FSceneOutlinerTreeItemPtr Item, ESelectInfo::Type SelectionType, const FSceneOutlinerItemSelection& Selection) override;
+	UE_API virtual void OnItemDoubleClick(FSceneOutlinerTreeItemPtr Item) override;
+	UE_API virtual void OnFilterTextCommited(FSceneOutlinerItemSelection& Selection, ETextCommit::Type CommitType) override;
+	UE_API virtual void OnItemPassesFilters(const ISceneOutlinerTreeItem& Item) override;
+	UE_API virtual FReply OnKeyDown(const FKeyEvent& InKeyEvent) override;
+	UE_API virtual void OnDuplicateSelected() override;
+	UE_API virtual void OnDrop(ISceneOutlinerTreeItem& DropTarget, const FSceneOutlinerDragDropPayload& Payload, const FSceneOutlinerDragValidationInfo& ValidationInfo) const override;
+	UE_API virtual bool CanRenameItem(const ISceneOutlinerTreeItem& Item) const override;
+	UE_API virtual FText GetStatusText() const override;
+	UE_API virtual FSlateColor GetStatusTextColor() const override;
 	virtual ESelectionMode::Type GetSelectionMode() const override { return ESelectionMode::Multi; }
 	virtual bool SupportsKeyboardFocus() const override { return true; }
 	virtual bool ShouldShowFolders() const override { return GetTreeViewMode() == EObjectMixerTreeViewMode::Folders; }
@@ -216,158 +218,158 @@ public:
 	virtual bool ShowStatusBar() const override { return true; }
 	virtual bool ShowViewButton() const override { return true; }
 	virtual bool ShowFilterOptions() const override { return true; }
-	virtual bool CanDelete() const override;
-	virtual bool CanRename() const override;
-	virtual bool CanCut() const override;
-	virtual bool CanCopy() const override;
-	virtual bool CanPaste() const override;
+	UE_API virtual bool CanDelete() const override;
+	UE_API virtual bool CanRename() const override;
+	UE_API virtual bool CanCut() const override;
+	UE_API virtual bool CanCopy() const override;
+	UE_API virtual bool CanPaste() const override;
 	virtual bool CanSupportDragAndDrop() const override { return true; }
 	virtual bool CanCustomizeToolbar() const override { return true; }
-	virtual bool HasErrors() const override;
-	virtual FText GetErrorsText() const override;
-	virtual void RepairErrors() const override;	
-	virtual FFolder CreateNewFolder() override;
-	virtual FFolder GetFolder(const FFolder& ParentPath, const FName& LeafName) override;
-	virtual bool CreateFolder(const FFolder& NewFolder) override;
-	virtual bool ReparentItemToFolder(const FFolder& FolderPath, const FSceneOutlinerTreeItemPtr& Item) override;
-	virtual void SelectFoldersDescendants(const TArray<FFolderTreeItem*>& FolderItems, bool bSelectImmediateChildrenOnly) override;
-	virtual void PinItems(const TArray<FSceneOutlinerTreeItemPtr>& InItems) override;
-	virtual void UnpinItems(const TArray<FSceneOutlinerTreeItemPtr>& InItems) override;
-	virtual bool CanPinItems(const TArray<FSceneOutlinerTreeItemPtr>& InItems) const override;
-	virtual bool CanUnpinItems(const TArray<FSceneOutlinerTreeItemPtr>& InItems) const override;
-	virtual void SynchronizeSelection() override;
+	UE_API virtual bool HasErrors() const override;
+	UE_API virtual FText GetErrorsText() const override;
+	UE_API virtual void RepairErrors() const override;	
+	UE_API virtual FFolder CreateNewFolder() override;
+	UE_API virtual FFolder GetFolder(const FFolder& ParentPath, const FName& LeafName) override;
+	UE_API virtual bool CreateFolder(const FFolder& NewFolder) override;
+	UE_API virtual bool ReparentItemToFolder(const FFolder& FolderPath, const FSceneOutlinerTreeItemPtr& Item) override;
+	UE_API virtual void SelectFoldersDescendants(const TArray<FFolderTreeItem*>& FolderItems, bool bSelectImmediateChildrenOnly) override;
+	UE_API virtual void PinItems(const TArray<FSceneOutlinerTreeItemPtr>& InItems) override;
+	UE_API virtual void UnpinItems(const TArray<FSceneOutlinerTreeItemPtr>& InItems) override;
+	UE_API virtual bool CanPinItems(const TArray<FSceneOutlinerTreeItemPtr>& InItems) const override;
+	UE_API virtual bool CanUnpinItems(const TArray<FSceneOutlinerTreeItemPtr>& InItems) const override;
+	UE_API virtual void SynchronizeSelection() override;
 	/* End ISceneOutlinerMode Interface */
 
 	/* External events this mode must respond to */
 
 	/** Called by the engine when a component is updated */
-	void OnComponentsUpdated();
+	UE_API void OnComponentsUpdated();
 	/** Called by the engine when an actor is deleted */
-	void OnLevelActorDeleted(AActor* Actor);
+	UE_API void OnLevelActorDeleted(AActor* Actor);
 
 	/** Called by the editor to allow selection of unloaded actors */
-	void OnSelectUnloadedActors(const TArray<FGuid>& ActorGuids);
+	UE_API void OnSelectUnloadedActors(const TArray<FGuid>& ActorGuids);
 	
 	/** Called when an actor desc instance is removed */
-	void OnActorDescInstanceRemoved(FWorldPartitionActorDescInstance* InActorDescInstance);
+	UE_API void OnActorDescInstanceRemoved(FWorldPartitionActorDescInstance* InActorDescInstance);
 	
 	UE_DEPRECATED(5.4, "Use OnActorDescInstanceRemoved instead")
 	void OnActorDescRemoved(FWorldPartitionActorDesc* InActorDesc) {}
 
 	/** Called by engine when edit cut actors begins */
-	void OnEditCutActorsBegin();
+	UE_API void OnEditCutActorsBegin();
 
 	/** Called by engine when edit cut actors ends */
-	void OnEditCutActorsEnd();
+	UE_API void OnEditCutActorsEnd();
 
 	/** Called by engine when edit copy actors begins */
-	void OnEditCopyActorsBegin();
+	UE_API void OnEditCopyActorsBegin();
 
 	/** Called by engine when edit copy actors ends */
-	void OnEditCopyActorsEnd();
+	UE_API void OnEditCopyActorsEnd();
 
 	/** Called by engine when edit paste actors begins */
-	void OnEditPasteActorsBegin();
+	UE_API void OnEditPasteActorsBegin();
 
 	/** Called by engine when edit paste actors ends */
-	void OnEditPasteActorsEnd();
+	UE_API void OnEditPasteActorsEnd();
 
 	/** Called by engine when edit duplicate actors begins */
-	void OnDuplicateActorsBegin();
+	UE_API void OnDuplicateActorsBegin();
 
 	/** Called by engine when edit duplicate actors ends */
-	void OnDuplicateActorsEnd();
+	UE_API void OnDuplicateActorsEnd();
 
 	/** Called by engine when edit delete actors begins */
-	void OnDeleteActorsBegin();
+	UE_API void OnDeleteActorsBegin();
 
 	/** Called by engine when edit delete actors ends */
-	void OnDeleteActorsEnd();
+	UE_API void OnDeleteActorsEnd();
 	
 	/** Function called by the Outliner Filter Bar to compare an item with Type Filters*/
-	virtual bool CompareItemWithClassName(SceneOutliner::FilterBarType InItem, const TSet<FTopLevelAssetPath>&) const override;
+	UE_API virtual bool CompareItemWithClassName(SceneOutliner::FilterBarType InItem, const TSet<FTopLevelAssetPath>&) const override;
 
 	/** Check whether hybrid rows are allowed for the associated object list */
-	bool ShouldAllowHybridRows() const;
+	UE_API bool ShouldAllowHybridRows() const;
 	
 protected:
 
 	/* Events */
 
-	void OnMapChange(uint32 MapFlags);
-	void OnNewCurrentLevel();
-
-	void OnEditorSelectionChanged();
-	void OnActorLabelChanged(AActor* ChangedActor);
-	void OnObjectsReplaced(const TMap<UObject*, UObject*>& ReplacementMap);
-	void OnLevelActorRequestsRename(const AActor* Actor);
-	void OnPostLoadMapWithWorld(UWorld* World);
-
-	bool ShouldSyncSelectionToEditor();
-	bool ShouldSyncSelectionFromEditor();
-	void SynchronizeAllSelectionsToEditor();
-	bool HasActorSelectionChanged(TArray<AActor*>& OutSelectedActors, bool& bOutAreAnyInPIE);
-	bool HasComponentSelectionChanged(TArray<UActorComponent*>& OutSelectedComponents, bool& bOutAreAnyInPIE);
-	void SelectActorsInEditor(const TArray<AActor*>& InSelectedActors, bool bShouldSelect, bool bSelectEvenIfHidden);
-	void SelectComponentsInEditor(const TArray<UActorComponent*>& InSelectedComponents, bool bShouldSelect, bool bSelectEvenIfHidden);
-	void SelectActorsInMixer(const TArray<AActor*>& InSelectedActors, bool bShouldSelect, bool bSelectEvenIfHidden);
-	void SelectComponentsInMixer(const TArray<UActorComponent*>& InSelectedComponents, bool bShouldSelect, bool bSelectEvenIfHidden);
+	UE_API void OnMapChange(uint32 MapFlags);
+	UE_API void OnNewCurrentLevel();
+
+	UE_API void OnEditorSelectionChanged();
+	UE_API void OnActorLabelChanged(AActor* ChangedActor);
+	UE_API void OnObjectsReplaced(const TMap<UObject*, UObject*>& ReplacementMap);
+	UE_API void OnLevelActorRequestsRename(const AActor* Actor);
+	UE_API void OnPostLoadMapWithWorld(UWorld* World);
+
+	UE_API bool ShouldSyncSelectionToEditor();
+	UE_API bool ShouldSyncSelectionFromEditor();
+	UE_API void SynchronizeAllSelectionsToEditor();
+	UE_API bool HasActorSelectionChanged(TArray<AActor*>& OutSelectedActors, bool& bOutAreAnyInPIE);
+	UE_API bool HasComponentSelectionChanged(TArray<UActorComponent*>& OutSelectedComponents, bool& bOutAreAnyInPIE);
+	UE_API void SelectActorsInEditor(const TArray<AActor*>& InSelectedActors, bool bShouldSelect, bool bSelectEvenIfHidden);
+	UE_API void SelectComponentsInEditor(const TArray<UActorComponent*>& InSelectedComponents, bool bShouldSelect, bool bSelectEvenIfHidden);
+	UE_API void SelectActorsInMixer(const TArray<AActor*>& InSelectedActors, bool bShouldSelect, bool bSelectEvenIfHidden);
+	UE_API void SelectComponentsInMixer(const TArray<UActorComponent*>& InSelectedComponents, bool bShouldSelect, bool bSelectEvenIfHidden);
 
 	/** Get the list of actors selected in the editor */
-	TArray<AActor*> GetSelectedActorsInEditor() const;
+	UE_API TArray<AActor*> GetSelectedActorsInEditor() const;
 
 	/** Get the list of components selected in the editor */
-	TArray<UActorComponent*> GetSelectedComponentsInEditor() const;
+	UE_API TArray<UActorComponent*> GetSelectedComponentsInEditor() const;
 
 	/** Check if an object is a member of one of the classes our list model is filtered to */
-	bool GetIsObjectOfFilteredClass(const UObject* Object) const;
+	UE_API bool GetIsObjectOfFilteredClass(const UObject* Object) const;
 	
 	/** Build and up the context menu */
-	TSharedPtr<SWidget> BuildContextMenu();
+	UE_API TSharedPtr<SWidget> BuildContextMenu();
 	/** Register the context menu with the engine */
-	static void RegisterContextMenu();
-	bool CanPasteFoldersOnlyFromClipboard() const;
+	static UE_API void RegisterContextMenu();
+	UE_API bool CanPasteFoldersOnlyFromClipboard() const;
 
-	bool GetFolderNamesFromPayload(const FSceneOutlinerDragDropPayload& InPayload, TArray<FName>& OutFolders, FFolder::FRootObject& OutCommonRootObject) const;
-	FFolder GetWorldDefaultRootFolder() const;
+	UE_API bool GetFolderNamesFromPayload(const FSceneOutlinerDragDropPayload& InPayload, TArray<FName>& OutFolders, FFolder::FRootObject& OutCommonRootObject) const;
+	UE_API FFolder GetWorldDefaultRootFolder() const;
 
 	/** Synchronize both the actor and component selection with the editor, with components taking priority over actors */
-	void SynchronizeComponentAndActorSelection();
-	void SynchronizeSelectedActorDescs();
+	UE_API void SynchronizeComponentAndActorSelection();
+	UE_API void SynchronizeSelectedActorDescs();
 
-	void OnActorEditorContextSubsystemChanged();
+	UE_API void OnActorEditorContextSubsystemChanged();
 
 	/** Filter factories */
-	static TSharedRef<FSceneOutlinerFilter> CreateShowOnlySelectedActorsFilter();
-	static TSharedRef<FSceneOutlinerFilter> CreateHideTemporaryActorsFilter();
-	static TSharedRef<FSceneOutlinerFilter> CreateIsInCurrentLevelFilter();
-	static TSharedRef<FSceneOutlinerFilter> CreateIsInCurrentDataLayersFilter();
-	static TSharedRef<FSceneOutlinerFilter> CreateHideComponentsFilter();
-	static TSharedRef<FSceneOutlinerFilter> CreateHideLevelInstancesFilter();
-	static TSharedRef<FSceneOutlinerFilter> CreateHideUnloadedActorsFilter();
-	static TSharedRef<FSceneOutlinerFilter> CreateHideEmptyFoldersFilter();
-	TSharedRef<FSceneOutlinerFilter> CreateIsInCurrentContentBundleFilter();
+	static UE_API TSharedRef<FSceneOutlinerFilter> CreateShowOnlySelectedActorsFilter();
+	static UE_API TSharedRef<FSceneOutlinerFilter> CreateHideTemporaryActorsFilter();
+	static UE_API TSharedRef<FSceneOutlinerFilter> CreateIsInCurrentLevelFilter();
+	static UE_API TSharedRef<FSceneOutlinerFilter> CreateIsInCurrentDataLayersFilter();
+	static UE_API TSharedRef<FSceneOutlinerFilter> CreateHideComponentsFilter();
+	static UE_API TSharedRef<FSceneOutlinerFilter> CreateHideLevelInstancesFilter();
+	static UE_API TSharedRef<FSceneOutlinerFilter> CreateHideUnloadedActorsFilter();
+	static UE_API TSharedRef<FSceneOutlinerFilter> CreateHideEmptyFoldersFilter();
+	UE_API TSharedRef<FSceneOutlinerFilter> CreateIsInCurrentContentBundleFilter();
 
 	/** Functions to expose selection framing to the UI */
-	void OnToggleAlwaysFrameSelection();
-	bool ShouldAlwaysFrameSelection();
+	UE_API void OnToggleAlwaysFrameSelection();
+	UE_API bool ShouldAlwaysFrameSelection();
 
 	/**
 	 * Get a mutable version of the ActorBrowser config for setting values.
 	 * @returns		The config for this ActorBrowser.
 	 * @note		If OutlinerIdentifier is not set for this outliner, it is not possible to store settings.
 	 */
-	FObjectMixerOutlinerModeConfig* GetMutableConfig() const;
+	UE_API FObjectMixerOutlinerModeConfig* GetMutableConfig() const;
 
 	/**
 	 * Get a const version of the ActorBrowser config for getting values.
 	 * @returns		The config for this ActorBrowser.
 	 * @note		If OutlinerIdentifier is not set for this outliner, it is not possible to store settings.
 	 */
-	const FObjectMixerOutlinerModeConfig* GetConstConfig() const;
+	UE_API const FObjectMixerOutlinerModeConfig* GetConstConfig() const;
 
 	/** Save the config for this ActorBrowser */
-	void SaveConfig();
+	UE_API void SaveConfig();
 
 	IWorldPartitionEditorModule* WorldPartitionEditorModule;
 	/** Number of actors (including unloaded) which have passed through the filters */
@@ -385,7 +387,7 @@ protected:
 	TWeakPtr<FObjectMixerEditorList> ListModelPtr;
 	
 	TArray<FFilterClassSelectionInfo> FilterClassSelectionInfos;
-	TSharedRef<SWidget> OnGenerateFilterClassMenu();
+	UE_API TSharedRef<SWidget> OnGenerateFilterClassMenu();
 
 	/** Used in case of a selection sync override. */
 	
@@ -395,3 +397,5 @@ protected:
 	/** The last selected tree item IDs, stored here so we can maintain the selection across refreshes when not synced with the editor. */
 	TArray<FSceneOutlinerTreeItemID> SelectedIDs;
 };
+
+#undef UE_API
diff --git a/Engine/Plugins/Editor/ObjectMixer/ObjectMixer/Source/ObjectMixer/Public/Views/List/Modes/SFilterClassMenuItem.h b/Engine/Plugins/Editor/ObjectMixer/ObjectMixer/Source/ObjectMixer/Public/Views/List/Modes/SFilterClassMenuItem.h
--- a/Engine/Plugins/Editor/ObjectMixer/ObjectMixer/Source/ObjectMixer/Public/Views/List/Modes/SFilterClassMenuItem.h
+++ b/Engine/Plugins/Editor/ObjectMixer/ObjectMixer/Source/ObjectMixer/Public/Views/List/Modes/SFilterClassMenuItem.h
@@ -22,7 +22,7 @@ struct FAssetClassMap
 	}
 };
 
-class OBJECTMIXEREDITOR_API SFilterClassMenuItem : public SCompoundWidget
+class SFilterClassMenuItem : public SCompoundWidget
 {
 public:
 
diff --git a/Engine/Plugins/Editor/ObjectMixer/ObjectMixer/Source/ObjectMixer/Public/Views/List/ObjectMixerEditorList.h b/Engine/Plugins/Editor/ObjectMixer/ObjectMixer/Source/ObjectMixer/Public/Views/List/ObjectMixerEditorList.h
--- a/Engine/Plugins/Editor/ObjectMixer/ObjectMixer/Source/ObjectMixer/Public/Views/List/ObjectMixerEditorList.h
+++ b/Engine/Plugins/Editor/ObjectMixer/ObjectMixer/Source/ObjectMixer/Public/Views/List/ObjectMixerEditorList.h
@@ -8,6 +8,8 @@
 #include "Templates/SharedPointer.h"
 #include "Widgets/SWidget.h"
 
+#define UE_API OBJECTMIXEREDITOR_API
+
 class FObjectMixerEditorModule;
 class SObjectMixerEditorList;
 class UObjectMixerEditorSerializedData;
@@ -16,56 +18,56 @@ DECLARE_MULTICAST_DELEGATE(FOnPreFilterChange)
 DECLARE_MULTICAST_DELEGATE(FOnPostFilterChange)
 
 
-class OBJECTMIXEREDITOR_API FObjectMixerEditorList : public TSharedFromThis<FObjectMixerEditorList>, public FGCObject
+class FObjectMixerEditorList : public TSharedFromThis<FObjectMixerEditorList>, public FGCObject
 {
 public:
 
-	FObjectMixerEditorList(const FName InModuleName, TSharedPtr<IObjectMixerSelectionInterface> InSelectionInterface = nullptr);
+	UE_API FObjectMixerEditorList(const FName InModuleName, TSharedPtr<IObjectMixerSelectionInterface> InSelectionInterface = nullptr);
 
-	virtual ~FObjectMixerEditorList();
+	UE_API virtual ~FObjectMixerEditorList();
 
-	void Initialize();
+	UE_API void Initialize();
 	
-	void RegisterAndMapContextMenuCommands();
+	UE_API void RegisterAndMapContextMenuCommands();
 
 	/** Creates a new widget only if one does not already exist. Otherwise returns the existing widget. */
-	TSharedRef<SWidget> GetOrCreateWidget();
+	UE_API TSharedRef<SWidget> GetOrCreateWidget();
 
 	/** Creates a new widget, replacing the existing one's pointer. Not useful to call alone, use RequestRegenerateListWidget */
-	TSharedRef<SWidget> CreateWidget();
+	UE_API TSharedRef<SWidget> CreateWidget();
 
 	/** Calls back to the module to replace the widget in the dock tab. Call when columns or the object filters have changed. */
-	void RequestRegenerateListWidget();
+	UE_API void RequestRegenerateListWidget();
 
 	/** Rebuild the list items and refresh the list. Call when adding or removing items. */
-	void RequestRebuildList() const;
+	UE_API void RequestRebuildList() const;
 
 	/**
 	 * Refresh list filters and sorting.
 	 * Useful for when the list state has gone stale but the variable count has not changed.
 	 */
-	void RefreshList() const;
+	UE_API void RefreshList() const;
 
 	/** Get the selection interface used to synchronize this with the rest of the editor */
 	TSharedPtr<IObjectMixerSelectionInterface> GetSelectionInterface() const { return SelectionInterface; }
 
-	void BuildPerformanceCache();
+	UE_API void BuildPerformanceCache();
 
-	bool ShouldShowTransientObjects() const;
+	UE_API bool ShouldShowTransientObjects() const;
 
 	/** Called when the Rename command is executed from the UI or hotkey. */
-	void OnRenameCommand();
+	UE_API void OnRenameCommand();
 	
-	void RebuildCollectionSelector();
+	UE_API void RebuildCollectionSelector();
 
-	void SetDefaultFilterClass(UClass* InNewClass);
-	bool IsClassSelected(UClass* InClass) const;
+	UE_API void SetDefaultFilterClass(UClass* InNewClass);
+	UE_API bool IsClassSelected(UClass* InClass) const;
 
-	const TArray<TObjectPtr<UObjectMixerObjectFilter>>& GetObjectFilterInstances();
+	UE_API const TArray<TObjectPtr<UObjectMixerObjectFilter>>& GetObjectFilterInstances();
 
-	const UObjectMixerObjectFilter* GetMainObjectFilterInstance();
+	UE_API const UObjectMixerObjectFilter* GetMainObjectFilterInstance();
 
-	void CacheObjectFilterInstances();
+	UE_API void CacheObjectFilterInstances();
 
 	/** Get the style of the tree (flat list or hierarchy) */
 	EObjectMixerTreeViewMode GetTreeViewMode()
@@ -83,22 +85,22 @@ public:
 	 * Force returns result from Filter->ForceGetObjectClassesToFilter.
 	 * Generally, you want to get this from the public performance cache.
 	 */
-	TSet<UClass*> ForceGetObjectClassesToFilter();
+	UE_API TSet<UClass*> ForceGetObjectClassesToFilter();
 
 	/**
 	 * Force returns result from Filter->ForceGetObjectClassesToPlace.
 	 * Generally, you want to get this from the public performance cache.
 	 */
-	TSet<TSubclassOf<AActor>> ForceGetObjectClassesToPlace();
+	UE_API TSet<TSubclassOf<AActor>> ForceGetObjectClassesToPlace();
 
 	const TArray<TSubclassOf<UObjectMixerObjectFilter>>& GetObjectFilterClasses() const
 	{
 		return ObjectFilterClasses;
 	}
 
-	void AddObjectFilterClass(UClass* InObjectFilterClass, const bool bShouldRebuild = true);
+	UE_API void AddObjectFilterClass(UClass* InObjectFilterClass, const bool bShouldRebuild = true);
 
-	void RemoveObjectFilterClass(UClass* InObjectFilterClass, const bool bCacheAndRebuild = true);
+	UE_API void RemoveObjectFilterClass(UClass* InObjectFilterClass, const bool bCacheAndRebuild = true);
 
 	void ResetObjectFilterClasses(const bool bCacheAndRebuild = true)
 	{
@@ -129,30 +131,30 @@ public:
 		return ModuleName;
 	}
 
-	FObjectMixerEditorModule* GetModulePtr() const;
+	UE_API FObjectMixerEditorModule* GetModulePtr() const;
 
 	// User Collections
 
 	/** Get a pointer to the UObjectMixerEditorSerializedData object along with the name of the filter represented by this ListModel instance. */
-	UObjectMixerEditorSerializedData* GetSerializedData() const;
-	bool RequestAddObjectsToCollection(const FName& CollectionName, const TSet<FSoftObjectPath>& ObjectsToAdd) const;
-	bool RequestRemoveObjectsFromCollection(const FName& CollectionName, const TSet<FSoftObjectPath>& ObjectsToRemove) const;
-	bool RequestRemoveCollection(const FName& CollectionName) const;
-	bool RequestDuplicateCollection(const FName& CollectionToDuplicateName, FName& DesiredDuplicateName) const;
-	bool RequestReorderCollection(const FName& CollectionToMoveName, const FName& CollectionInsertBeforeName) const;
-	bool RequestRenameCollection(const FName& CollectionNameToRename, const FName& NewCollectionName) const;
-	bool DoesCollectionExist(const FName& CollectionName) const;
-	bool IsObjectInCollection(const FName& CollectionName, const FSoftObjectPath& InObject) const;
-	TSet<FName> GetCollectionsForObject(const FSoftObjectPath& InObject) const;
-	TArray<FName> GetAllCollectionNames() const;
+	UE_API UObjectMixerEditorSerializedData* GetSerializedData() const;
+	UE_API bool RequestAddObjectsToCollection(const FName& CollectionName, const TSet<FSoftObjectPath>& ObjectsToAdd) const;
+	UE_API bool RequestRemoveObjectsFromCollection(const FName& CollectionName, const TSet<FSoftObjectPath>& ObjectsToRemove) const;
+	UE_API bool RequestRemoveCollection(const FName& CollectionName) const;
+	UE_API bool RequestDuplicateCollection(const FName& CollectionToDuplicateName, FName& DesiredDuplicateName) const;
+	UE_API bool RequestReorderCollection(const FName& CollectionToMoveName, const FName& CollectionInsertBeforeName) const;
+	UE_API bool RequestRenameCollection(const FName& CollectionNameToRename, const FName& NewCollectionName) const;
+	UE_API bool DoesCollectionExist(const FName& CollectionName) const;
+	UE_API bool IsObjectInCollection(const FName& CollectionName, const FSoftObjectPath& InObject) const;
+	UE_API TSet<FName> GetCollectionsForObject(const FSoftObjectPath& InObject) const;
+	UE_API TArray<FName> GetAllCollectionNames() const;
 
 	/**
 	 * This is the filter class used to initialize the ListModel.
 	 * This filter class cannot be turned off by the end user.
 	 */
-	const TSubclassOf<UObjectMixerObjectFilter>& GetDefaultFilterClass() const;
+	UE_API const TSubclassOf<UObjectMixerObjectFilter>& GetDefaultFilterClass() const;
 
-	void OnPostFilterChange();
+	UE_API void OnPostFilterChange();
 
 	FOnPreFilterChange OnPreFilterChangeDelegate;
 	FOnPostFilterChange OnPostFilterChangeDelegate;
@@ -160,28 +162,28 @@ public:
 	TSharedPtr<FUICommandList> ObjectMixerElementEditCommands;
 	TSharedPtr<FUICommandList> ObjectMixerFolderEditCommands;
 
-	void FlushWidget();
+	UE_API void FlushWidget();
 
-	[[nodiscard]] TArray<TSharedPtr<ISceneOutlinerTreeItem>> GetSelectedTreeViewItems() const;
-	int32 GetSelectedTreeViewItemCount() const;
+	[[nodiscard]] UE_API TArray<TSharedPtr<ISceneOutlinerTreeItem>> GetSelectedTreeViewItems() const;
+	UE_API int32 GetSelectedTreeViewItemCount() const;
 
-	TSet<TSharedPtr<ISceneOutlinerTreeItem>> GetSoloRows() const;
-	void ClearSoloRows();
+	UE_API TSet<TSharedPtr<ISceneOutlinerTreeItem>> GetSoloRows() const;
+	UE_API void ClearSoloRows();
 
 	/** Returns true if at least one row is set to Solo. */
-	bool IsListInSoloState() const;
+	UE_API bool IsListInSoloState() const;
 
 	/**
 	 * Determines whether rows' objects should be temporarily hidden in editor based on each row's visibility rules,
 	 * then sets each object's visibility in editor.
 	 */
-	void EvaluateAndSetEditorVisibilityPerRow();
+	UE_API void EvaluateAndSetEditorVisibilityPerRow();
 	
 	/** Represents the collections the user has selected in the UI. If empty, "All" is considered as selected. */
-	[[nodiscard]] const TSet<FName>& GetSelectedCollections() const;
-	[[nodiscard]] bool IsCollectionSelected(const FName& CollectionName) const;
-	void SetSelectedCollections(const TSet<FName> InSelectedCollections);
-	void SetCollectionSelected(const FName& CollectionName, const bool bNewSelected);
+	[[nodiscard]] UE_API const TSet<FName>& GetSelectedCollections() const;
+	[[nodiscard]] UE_API bool IsCollectionSelected(const FName& CollectionName) const;
+	UE_API void SetSelectedCollections(const TSet<FName> InSelectedCollections);
+	UE_API void SetCollectionSelected(const FName& CollectionName, const bool bNewSelected);
 
 	// Performance cache
 	TSet<UClass*> ObjectClassesToFilterCache;
@@ -195,13 +197,13 @@ public:
 protected:
 
 	/** Bind delegates to refresh the list when editor state changes. */
-	virtual void BindRefreshDelegates();
+	UE_API virtual void BindRefreshDelegates();
 
 	/** Unbind delegates to refresh the list when editor state changes. */
-	virtual void UnbindRefreshDelegates();
+	UE_API virtual void UnbindRefreshDelegates();
 
 	/** If a property is changed that has a name found in this set, the list will be refreshed. */
-	TSet<FName> GetPropertiesThatRequireRefresh();
+	UE_API TSet<FName> GetPropertiesThatRequireRefresh();
 	
 	virtual void AddReferencedObjects( FReferenceCollector& Collector )  override
 	{
@@ -246,3 +248,5 @@ protected:
 	/** Handles for delegates registered to refresh the list */
 	TSet<FDelegateHandle> DelegateHandles;
 };
+
+#undef UE_API
diff --git a/Engine/Plugins/Editor/ObjectMixer/ObjectMixer/Source/ObjectMixer/Public/Views/List/ObjectMixerEditorListRowData.h b/Engine/Plugins/Editor/ObjectMixer/ObjectMixer/Source/ObjectMixer/Public/Views/List/ObjectMixerEditorListRowData.h
--- a/Engine/Plugins/Editor/ObjectMixer/ObjectMixer/Source/ObjectMixer/Public/Views/List/ObjectMixerEditorListRowData.h
+++ b/Engine/Plugins/Editor/ObjectMixer/ObjectMixer/Source/ObjectMixer/Public/Views/List/ObjectMixerEditorListRowData.h
@@ -8,10 +8,12 @@
 #include "PropertyHandle.h"
 #include "SSceneOutliner.h"
 
+#define UE_API OBJECTMIXEREDITOR_API
+
 class SObjectMixerEditorList;
 
 /** Defines data carried by each row type. */
-struct OBJECTMIXEREDITOR_API FObjectMixerEditorListRowData
+struct FObjectMixerEditorListRowData
 {
 	struct FTransientEditorVisibilityRules
 	{
@@ -67,17 +69,17 @@ struct OBJECTMIXEREDITOR_API FObjectMixerEditorListRowData
 		return SceneOutlinerPtr.IsValid();
 	}
 
-	const TArray<TObjectPtr<UObjectMixerObjectFilter>>& GetObjectFilterInstances() const;
+	UE_API const TArray<TObjectPtr<UObjectMixerObjectFilter>>& GetObjectFilterInstances() const;
 
-	const UObjectMixerObjectFilter* GetMainObjectFilterInstance() const;
+	UE_API const UObjectMixerObjectFilter* GetMainObjectFilterInstance() const;
 
-	[[nodiscard]] bool GetIsTreeViewItemExpanded(const TSharedRef<ISceneOutlinerTreeItem> InRow);
-	void SetIsTreeViewItemExpanded(const TSharedRef<ISceneOutlinerTreeItem> InRow, const bool bNewExpanded);
+	[[nodiscard]] UE_API bool GetIsTreeViewItemExpanded(const TSharedRef<ISceneOutlinerTreeItem> InRow);
+	UE_API void SetIsTreeViewItemExpanded(const TSharedRef<ISceneOutlinerTreeItem> InRow, const bool bNewExpanded);
 
-	[[nodiscard]] bool GetIsSelected(const TSharedRef<ISceneOutlinerTreeItem> InRow);
-	void SetIsSelected(const TSharedRef<ISceneOutlinerTreeItem> InRow, const bool bNewSelected);
+	[[nodiscard]] UE_API bool GetIsSelected(const TSharedRef<ISceneOutlinerTreeItem> InRow);
+	UE_API void SetIsSelected(const TSharedRef<ISceneOutlinerTreeItem> InRow, const bool bNewSelected);
 
-	[[nodiscard]] bool HasAtLeastOneChildThatIsNotSolo(
+	[[nodiscard]] UE_API bool HasAtLeastOneChildThatIsNotSolo(
 		const TSharedRef<ISceneOutlinerTreeItem> InRow, const bool bRecursive = true) const;
 
 	[[nodiscard]] const FText& GetDisplayNameOverride() const
@@ -90,22 +92,22 @@ struct OBJECTMIXEREDITOR_API FObjectMixerEditorListRowData
 		DisplayNameOverride = InDisplayNameOverride;
 	}
 
-	[[nodiscard]] TWeakPtr<SObjectMixerEditorList> GetListView() const;
+	[[nodiscard]] UE_API TWeakPtr<SObjectMixerEditorList> GetListView() const;
 
-	[[nodiscard]] TArray<TSharedPtr<ISceneOutlinerTreeItem>> GetSelectedTreeViewItems() const;
+	[[nodiscard]] UE_API TArray<TSharedPtr<ISceneOutlinerTreeItem>> GetSelectedTreeViewItems() const;
 	
-	void OnChangeVisibility(const FSceneOutlinerTreeItemRef TreeItem, const bool bNewVisible);
+	UE_API void OnChangeVisibility(const FSceneOutlinerTreeItemRef TreeItem, const bool bNewVisible);
 	
-	[[nodiscard]] const FTransientEditorVisibilityRules& GetVisibilityRules() const;
-	void SetVisibilityRules(const FTransientEditorVisibilityRules& InVisibilityRules);
+	[[nodiscard]] UE_API const FTransientEditorVisibilityRules& GetVisibilityRules() const;
+	UE_API void SetVisibilityRules(const FTransientEditorVisibilityRules& InVisibilityRules);
 
-	bool IsUserSetHiddenInEditor() const;
-	void SetUserHiddenInEditor(const bool bNewHidden);
+	UE_API bool IsUserSetHiddenInEditor() const;
+	UE_API void SetUserHiddenInEditor(const bool bNewHidden);
 	
-	bool GetRowSoloState() const;
-	void SetRowSoloState(const bool bNewSolo);
+	UE_API bool GetRowSoloState() const;
+	UE_API void SetRowSoloState(const bool bNewSolo);
 
-	void ClearSoloRows() const;
+	UE_API void ClearSoloRows() const;
 
 	bool GetIsHybridRow() const
 	{
@@ -126,7 +128,7 @@ struct OBJECTMIXEREDITOR_API FObjectMixerEditorListRowData
 	/** Attempt to propagate a property's new value to other selected rows.
 	 * @return Returns true if the property's handle was found or there is no work to do (i.e., no property name specified or this row is not selected)
 	 */
-	bool PropagateChangesToSimilarSelectedRowProperties(
+	UE_API bool PropagateChangesToSimilarSelectedRowProperties(
 		const TSharedRef<ISceneOutlinerTreeItem> InRow, const FPropertyPropagationInfo PropertyPropagationInfo);
 	
 	TMap<FName, TWeakPtr<IPropertyHandle>> PropertyNamesToHandles;
@@ -145,3 +147,5 @@ protected:
 	 */
 	TSoftObjectPtr<UActorComponent> HybridComponentPath = nullptr;
 };
+
+#undef UE_API
diff --git a/Engine/Plugins/Editor/ObjectMixer/ObjectMixer/Source/ObjectMixer/Public/Views/List/RowTypes/ObjectMixerEditorListRowActor.h b/Engine/Plugins/Editor/ObjectMixer/ObjectMixer/Source/ObjectMixer/Public/Views/List/RowTypes/ObjectMixerEditorListRowActor.h
--- a/Engine/Plugins/Editor/ObjectMixer/ObjectMixer/Source/ObjectMixer/Public/Views/List/RowTypes/ObjectMixerEditorListRowActor.h
+++ b/Engine/Plugins/Editor/ObjectMixer/ObjectMixer/Source/ObjectMixer/Public/Views/List/RowTypes/ObjectMixerEditorListRowActor.h
@@ -8,7 +8,9 @@
 #include "GameFramework/Actor.h"
 #include "SSceneOutliner.h"
 
-struct OBJECTMIXEREDITOR_API FObjectMixerEditorListRowActor : FActorTreeItem
+#define UE_API OBJECTMIXEREDITOR_API
+
+struct FObjectMixerEditorListRowActor : FActorTreeItem
 {
 	explicit FObjectMixerEditorListRowActor(
 		AActor* InObject, 
@@ -32,8 +34,10 @@ struct OBJECTMIXEREDITOR_API FObjectMixerEditorListRowActor : FActorTreeItem
 	TSoftObjectPtr<AActor> OverrideParent;
 
 	/* Begin ISceneOutlinerTreeItem Implementation */
-	static const FSceneOutlinerTreeItemType Type;
-	virtual void OnVisibilityChanged(const bool bNewVisibility) override;
-	virtual FSceneOutlinerTreeItemID GetID() const override;
+	static UE_API const FSceneOutlinerTreeItemType Type;
+	UE_API virtual void OnVisibilityChanged(const bool bNewVisibility) override;
+	UE_API virtual FSceneOutlinerTreeItemID GetID() const override;
 	/* End ISceneOutlinerTreeItem Implementation */
 };
+
+#undef UE_API
diff --git a/Engine/Plugins/Editor/ObjectMixer/ObjectMixer/Source/ObjectMixer/Public/Views/List/RowTypes/ObjectMixerEditorListRowComponent.h b/Engine/Plugins/Editor/ObjectMixer/ObjectMixer/Source/ObjectMixer/Public/Views/List/RowTypes/ObjectMixerEditorListRowComponent.h
--- a/Engine/Plugins/Editor/ObjectMixer/ObjectMixer/Source/ObjectMixer/Public/Views/List/RowTypes/ObjectMixerEditorListRowComponent.h
+++ b/Engine/Plugins/Editor/ObjectMixer/ObjectMixer/Source/ObjectMixer/Public/Views/List/RowTypes/ObjectMixerEditorListRowComponent.h
@@ -7,7 +7,9 @@
 #include "ComponentTreeItem.h"
 #include "SSceneOutliner.h"
 
-struct OBJECTMIXEREDITOR_API FObjectMixerEditorListRowComponent : FComponentTreeItem
+#define UE_API OBJECTMIXEREDITOR_API
+
+struct FObjectMixerEditorListRowComponent : FComponentTreeItem
 {
 	explicit FObjectMixerEditorListRowComponent(
 		UActorComponent* InObject, 
@@ -25,6 +27,8 @@ struct OBJECTMIXEREDITOR_API FObjectMixerEditorListRowComponent : FComponentTree
 	TSoftObjectPtr<UActorComponent> OriginalObjectSoftPtr;
 	
 	/* Begin ISceneOutlinerTreeItem Implementation */
-	static const FSceneOutlinerTreeItemType Type;
+	static UE_API const FSceneOutlinerTreeItemType Type;
 	/* End ISceneOutlinerTreeItem Implementation */
 };
+
+#undef UE_API
diff --git a/Engine/Plugins/Editor/ObjectMixer/ObjectMixer/Source/ObjectMixer/Public/Views/List/RowTypes/ObjectMixerEditorListRowFolder.h b/Engine/Plugins/Editor/ObjectMixer/ObjectMixer/Source/ObjectMixer/Public/Views/List/RowTypes/ObjectMixerEditorListRowFolder.h
--- a/Engine/Plugins/Editor/ObjectMixer/ObjectMixer/Source/ObjectMixer/Public/Views/List/RowTypes/ObjectMixerEditorListRowFolder.h
+++ b/Engine/Plugins/Editor/ObjectMixer/ObjectMixer/Source/ObjectMixer/Public/Views/List/RowTypes/ObjectMixerEditorListRowFolder.h
@@ -9,7 +9,9 @@
 #include "Folder.h"
 #include "SSceneOutliner.h"
 
-struct OBJECTMIXEREDITOR_API FObjectMixerEditorListRowFolder : FActorFolderTreeItem
+#define UE_API OBJECTMIXEREDITOR_API
+
+struct FObjectMixerEditorListRowFolder : FActorFolderTreeItem
 {
 	explicit FObjectMixerEditorListRowFolder(
 		const FFolder InFolder, 
@@ -23,7 +25,9 @@ struct OBJECTMIXEREDITOR_API FObjectMixerEditorListRowFolder : FActorFolderTreeI
 	FObjectMixerEditorListRowData RowData;
 
 	/* Begin ISceneOutlinerTreeItem Implementation */
-	static const FSceneOutlinerTreeItemType Type;
-	virtual void OnVisibilityChanged(const bool bNewVisibility) override;
+	static UE_API const FSceneOutlinerTreeItemType Type;
+	UE_API virtual void OnVisibilityChanged(const bool bNewVisibility) override;
 	/* End ISceneOutlinerTreeItem Implementation */
 };
+
+#undef UE_API
diff --git a/Engine/Plugins/Editor/ObjectMixer/ObjectMixer/Source/ObjectMixer/Public/Views/List/RowTypes/ObjectMixerEditorListRowUObject.h b/Engine/Plugins/Editor/ObjectMixer/ObjectMixer/Source/ObjectMixer/Public/Views/List/RowTypes/ObjectMixerEditorListRowUObject.h
--- a/Engine/Plugins/Editor/ObjectMixer/ObjectMixer/Source/ObjectMixer/Public/Views/List/RowTypes/ObjectMixerEditorListRowUObject.h
+++ b/Engine/Plugins/Editor/ObjectMixer/ObjectMixer/Source/ObjectMixer/Public/Views/List/RowTypes/ObjectMixerEditorListRowUObject.h
@@ -9,7 +9,9 @@
 #include "SSceneOutliner.h"
 #include "UObject/Object.h"
 
-struct OBJECTMIXEREDITOR_API FObjectMixerEditorListRowUObject : ISceneOutlinerTreeItem
+#define UE_API OBJECTMIXEREDITOR_API
+
+struct FObjectMixerEditorListRowUObject : ISceneOutlinerTreeItem
 {	
 	explicit FObjectMixerEditorListRowUObject(
 		UObject* InObject, 
@@ -31,10 +33,12 @@ struct OBJECTMIXEREDITOR_API FObjectMixerEditorListRowUObject : ISceneOutlinerTr
 	const FObjectKey ID;
 
 	/* Begin ISceneOutlinerTreeItem Implementation */
-	static const FSceneOutlinerTreeItemType Type;
-	virtual FString GetDisplayString() const override;
+	static UE_API const FSceneOutlinerTreeItemType Type;
+	UE_API virtual FString GetDisplayString() const override;
 	virtual bool CanInteract() const override { return true; }
 	virtual bool IsValid() const override { return !ObjectSoftPtr.IsNull(); }
 	virtual FSceneOutlinerTreeItemID GetID() const override { return ID; }
 	/* End ISceneOutlinerTreeItem Implementation */
 };
+
+#undef UE_API
diff --git a/Engine/Plugins/Editor/ObjectMixer/ObjectMixer/Source/ObjectMixer/Public/Views/List/SObjectMixerEditorList.h b/Engine/Plugins/Editor/ObjectMixer/ObjectMixer/Source/ObjectMixer/Public/Views/List/SObjectMixerEditorList.h
--- a/Engine/Plugins/Editor/ObjectMixer/ObjectMixer/Source/ObjectMixer/Public/Views/List/SObjectMixerEditorList.h
+++ b/Engine/Plugins/Editor/ObjectMixer/ObjectMixer/Source/ObjectMixer/Public/Views/List/SObjectMixerEditorList.h
@@ -10,6 +10,8 @@
 
 #include "SObjectMixerEditorList.generated.h"
 
+#define UE_API OBJECTMIXEREDITOR_API
+
 class FObjectMixerEditorList;
 class FObjectMixerOutlinerMode;
 class SBox;
@@ -67,47 +69,47 @@ struct FObjectMixerSceneOutlinerColumnInfo
 
 struct FObjectMixerEditorListRowData;
 
-class OBJECTMIXEREDITOR_API SObjectMixerEditorList : public SSceneOutliner
+class SObjectMixerEditorList : public SSceneOutliner
 {
 
 public:
 	DECLARE_EVENT(SObjectMixerEditorList, FSelectionSynchronizedEvent);
 	
 	// Columns
-	static const FName ItemNameColumnName;
-	static const FName EditorVisibilityColumnName;
-	static const FName EditorVisibilitySoloColumnName;
+	static UE_API const FName ItemNameColumnName;
+	static UE_API const FName EditorVisibilityColumnName;
+	static UE_API const FName EditorVisibilitySoloColumnName;
 
-	void Construct(const FArguments& InArgs, TSharedRef<FObjectMixerEditorList> ListModel);
+	UE_API void Construct(const FArguments& InArgs, TSharedRef<FObjectMixerEditorList> ListModel);
 
-	virtual ~SObjectMixerEditorList() override;
+	UE_API virtual ~SObjectMixerEditorList() override;
 
-	virtual FReply OnKeyDown( const FGeometry& MyGeometry, const FKeyEvent& InKeyEvent ) override;
+	UE_API virtual FReply OnKeyDown( const FGeometry& MyGeometry, const FKeyEvent& InKeyEvent ) override;
 
 	/**
 	 * Determines the style of the tree (flat list or hierarchy)
 	 */
-	EObjectMixerTreeViewMode GetTreeViewMode();
+	UE_API EObjectMixerTreeViewMode GetTreeViewMode();
 	/**
 	 * Determine the style of the tree (flat list or hierarchy)
 	 */
-	void SetTreeViewMode(EObjectMixerTreeViewMode InViewMode);
+	UE_API void SetTreeViewMode(EObjectMixerTreeViewMode InViewMode);
 
-	const TSet<FName>& GetSelectedCollections();
-	[[nodiscard]] bool IsCollectionSelected(const FName& CollectionName);
+	UE_API const TSet<FName>& GetSelectedCollections();
+	[[nodiscard]] UE_API bool IsCollectionSelected(const FName& CollectionName);
 	
-	void RebuildCollectionSelector();
+	UE_API void RebuildCollectionSelector();
 
-	bool RequestRemoveCollection(const FName& CollectionName);
-	bool RequestDuplicateCollection(const FName& CollectionToDuplicateName, FName& DesiredDuplicateName) const;
-	bool RequestRenameCollection(const FName& CollectionNameToRename, const FName& NewCollectionName);
-	bool DoesCollectionExist(const FName& CollectionName) const;
+	UE_API bool RequestRemoveCollection(const FName& CollectionName);
+	UE_API bool RequestDuplicateCollection(const FName& CollectionToDuplicateName, FName& DesiredDuplicateName) const;
+	UE_API bool RequestRenameCollection(const FName& CollectionNameToRename, const FName& NewCollectionName);
+	UE_API bool DoesCollectionExist(const FName& CollectionName) const;
 	
-	void OnCollectionCheckedStateChanged(bool bShouldBeChecked, FName CollectionName);
-	ECheckBoxState GetCollectionCheckedState(FName CollectionName) const;
+	UE_API void OnCollectionCheckedStateChanged(bool bShouldBeChecked, FName CollectionName);
+	UE_API ECheckBoxState GetCollectionCheckedState(FName CollectionName) const;
 
-	FObjectMixerOutlinerMode* GetCastedMode() const;
-	UWorld* GetWorld() const;
+	UE_API FObjectMixerOutlinerMode* GetCastedMode() const;
+	UE_API UWorld* GetWorld() const;
 
 	TWeakPtr<FObjectMixerEditorList> GetListModelPtr() const
 	{
@@ -118,56 +120,56 @@ public:
 	 * Refresh filters and sorting.
 	 * Useful for when the list state has gone stale but the variable count has not changed.
 	 */
-	void RefreshList();
+	UE_API void RefreshList();
 
 	/**
 	 * Regenerate the list items and refresh the list. Call when adding or removing items.
 	 */
-	void RequestRebuildList(const FString& InItemToScrollTo = "");
+	UE_API void RequestRebuildList(const FString& InItemToScrollTo = "");
 
 	/** Called when the Rename command is executed from the UI or hotkey. */
-	void OnRenameCommand();
+	UE_API void OnRenameCommand();
 
-	void AddToPendingPropertyPropagations(const FObjectMixerEditorListRowData::FPropertyPropagationInfo& InPropagationInfo);
+	UE_API void AddToPendingPropertyPropagations(const FObjectMixerEditorListRowData::FPropertyPropagationInfo& InPropagationInfo);
 
-	[[nodiscard]] TArray<TSharedPtr<ISceneOutlinerTreeItem>> GetSelectedTreeViewItems() const;
-	int32 GetSelectedTreeViewItemCount() const;
+	[[nodiscard]] UE_API TArray<TSharedPtr<ISceneOutlinerTreeItem>> GetSelectedTreeViewItems() const;
+	UE_API int32 GetSelectedTreeViewItemCount() const;
 
-	void SetTreeViewItemSelected(TSharedRef<ISceneOutlinerTreeItem> Item, const bool bNewSelected);
+	UE_API void SetTreeViewItemSelected(TSharedRef<ISceneOutlinerTreeItem> Item, const bool bNewSelected);
 
-	bool IsTreeViewItemSelected(TSharedPtr<ISceneOutlinerTreeItem> Item);
+	UE_API bool IsTreeViewItemSelected(TSharedPtr<ISceneOutlinerTreeItem> Item);
 
-	TSet<TSharedPtr<ISceneOutlinerTreeItem>> GetTreeRootItems() const;
-	TSet<TWeakPtr<ISceneOutlinerTreeItem>> GetWeakTreeRootItems() const;
+	UE_API TSet<TSharedPtr<ISceneOutlinerTreeItem>> GetTreeRootItems() const;
+	UE_API TSet<TWeakPtr<ISceneOutlinerTreeItem>> GetWeakTreeRootItems() const;
 	
-	TSet<TSharedPtr<ISceneOutlinerTreeItem>> GetSoloRows() const;
-	void ClearSoloRows();
+	UE_API TSet<TSharedPtr<ISceneOutlinerTreeItem>> GetSoloRows() const;
+	UE_API void ClearSoloRows();
 
 	/** Returns true if at least one row is set to Solo. */
-	bool IsListInSoloState() const;
+	UE_API bool IsListInSoloState() const;
 
 	/**
 	 * Determines whether rows' objects should be temporarily hidden in editor based on each row's visibility rules,
 	 * then sets each object's visibility in editor.
 	 */
-	void EvaluateAndSetEditorVisibilityPerRow();
+	UE_API void EvaluateAndSetEditorVisibilityPerRow();
 
-	bool IsTreeViewItemExpanded(const TSharedPtr<ISceneOutlinerTreeItem>& Row) const;
-	void SetTreeViewItemExpanded(const TSharedPtr<ISceneOutlinerTreeItem>& RowToExpand, const bool bNewExpansion) const;
+	UE_API bool IsTreeViewItemExpanded(const TSharedPtr<ISceneOutlinerTreeItem>& Row) const;
+	UE_API void SetTreeViewItemExpanded(const TSharedPtr<ISceneOutlinerTreeItem>& RowToExpand, const bool bNewExpansion) const;
 
-	void PropagatePropertyChangesToSelectedRows();
+	UE_API void PropagatePropertyChangesToSelectedRows();
 
 	// Columns
 
-	FObjectMixerSceneOutlinerColumnInfo* GetColumnInfoByPropertyName(const FName& InPropertyName);
-	void RestoreDefaultPropertyColumns();
-	TSharedRef<SWidget> GenerateHeaderRowContextMenu();
+	UE_API FObjectMixerSceneOutlinerColumnInfo* GetColumnInfoByPropertyName(const FName& InPropertyName);
+	UE_API void RestoreDefaultPropertyColumns();
+	UE_API TSharedRef<SWidget> GenerateHeaderRowContextMenu();
 
 	/** Get the event that broadcasts whenever the selection list is synchronized from the editor */
 	FSelectionSynchronizedEvent& GetOnSelectionSynchronized() { return OnSelectionSynchronizedEvent; }
 
 	/* Begin SSceneOutliner Interface */
-	virtual void CustomAddToToolbar(TSharedPtr<class SHorizontalBox> Toolbar) override;
+	UE_API virtual void CustomAddToToolbar(TSharedPtr<class SHorizontalBox> Toolbar) override;
 	/* End SSceneOutliner Interface */
 
 protected:
@@ -190,13 +192,13 @@ protected:
 
 	FSelectionSynchronizedEvent OnSelectionSynchronizedEvent;
 
-	bool CanCreateFolder() const;
-	TSharedRef<SWidget> OnGenerateAddObjectButtonMenu() const;
+	UE_API bool CanCreateFolder() const;
+	UE_API TSharedRef<SWidget> OnGenerateAddObjectButtonMenu() const;
 	
 	/** Disable all collection filters except CollectionToEnableName */
-	void SetSingleCollectionSelection(const FName& CollectionToEnableName = UObjectMixerEditorSerializedData::AllCollectionName);
+	UE_API void SetSingleCollectionSelection(const FName& CollectionToEnableName = UObjectMixerEditorSerializedData::AllCollectionName);
 
-	virtual void Tick(const FGeometry& AllottedGeometry, const double InCurrentTime, const float InDeltaTime) override;
+	UE_API virtual void Tick(const FGeometry& AllottedGeometry, const double InCurrentTime, const float InDeltaTime) override;
 
 	void OnActorSpawnedOrDestroyed(AActor* Object)
 	{
@@ -206,23 +208,25 @@ protected:
 	/**
 	 * Regenerate the list items and refresh the list. Call when adding or removing items.
 	 */
-	void RebuildList();
+	UE_API void RebuildList();
 
-	void InsertCollectionSelector();
+	UE_API void InsertCollectionSelector();
 
 	/**
 	 * Only adds properties that pass a series of tests, including having only one unique entry in the column list array.
 	 * @param bForceIncludeProperty If true, only Skiplist and Uniqueness tests will be checked, bypassing class, blueprint editability and other requirements.
 	 * @param PropertySkipList These property names will be skipped when they are encountered in the iteration.
 	 */
-	bool AddUniquePropertyColumnInfo(
+	UE_API bool AddUniquePropertyColumnInfo(
 		FProperty* Property,
 		const bool bForceIncludeProperty = false,
 		const TSet<FName>& PropertySkipList = {}
 	);
 
-	FText GetLevelColumnName() const;
-	void CreateActorTextInfoColumns(FSceneOutlinerInitializationOptions& OutInitOptions);
-	void SetupColumns(FSceneOutlinerInitializationOptions& OutInitOptions);
-	void SortColumnsForHeaderContextMenu();
+	UE_API FText GetLevelColumnName() const;
+	UE_API void CreateActorTextInfoColumns(FSceneOutlinerInitializationOptions& OutInitOptions);
+	UE_API void SetupColumns(FSceneOutlinerInitializationOptions& OutInitOptions);
+	UE_API void SortColumnsForHeaderContextMenu();
 };
+
+#undef UE_API
diff --git a/Engine/Plugins/Editor/ObjectMixer/ObjectMixer/Source/ObjectMixer/Public/Views/Widgets/ObjectMixerEditorListMenuContext.h b/Engine/Plugins/Editor/ObjectMixer/ObjectMixer/Source/ObjectMixer/Public/Views/Widgets/ObjectMixerEditorListMenuContext.h
--- a/Engine/Plugins/Editor/ObjectMixer/ObjectMixer/Source/ObjectMixer/Public/Views/Widgets/ObjectMixerEditorListMenuContext.h
+++ b/Engine/Plugins/Editor/ObjectMixer/ObjectMixer/Source/ObjectMixer/Public/Views/Widgets/ObjectMixerEditorListMenuContext.h
@@ -9,10 +9,12 @@
 
 #include "ObjectMixerEditorListMenuContext.generated.h"	
 
+#define UE_API OBJECTMIXEREDITOR_API
+
 class UToolMenu;
 
-UCLASS()
-class OBJECTMIXEREDITOR_API UObjectMixerEditorListMenuContext : public UObject
+UCLASS(MinimalAPI)
+class UObjectMixerEditorListMenuContext : public UObject
 {
 
 	GENERATED_BODY()
@@ -25,25 +27,27 @@ public:
 		TWeakPtr<class FObjectMixerEditorList> ListModelPtr;
 	};
 
-	static void RegisterFoldersOnlyContextMenu();
-	static void RegisterObjectMixerDynamicCollectionsContextMenuExtension(const FName& MenuName);
+	static UE_API void RegisterFoldersOnlyContextMenu();
+	static UE_API void RegisterObjectMixerDynamicCollectionsContextMenuExtension(const FName& MenuName);
 
 	FObjectMixerEditorListMenuContextData Data;
 
 private:
 
-	static bool DoesSelectionHaveType(const FObjectMixerEditorListMenuContextData& InData, UClass* Type);
-	static void AddCollectionWidget(const FName& Key, const FObjectMixerEditorListMenuContextData& ContextData, UToolMenu* Menu);
-	static void CreateSelectCollectionsSubMenu(UToolMenu* Menu, FObjectMixerEditorListMenuContextData ContextData);
-	static void AddCollectionsMenuItem(UToolMenu* InMenu, const FObjectMixerEditorListMenuContextData& ContextData);
-	static void OnTextCommitted(const FText& InText, ETextCommit::Type InCommitType, const FObjectMixerEditorListMenuContextData ContextData);
+	static UE_API bool DoesSelectionHaveType(const FObjectMixerEditorListMenuContextData& InData, UClass* Type);
+	static UE_API void AddCollectionWidget(const FName& Key, const FObjectMixerEditorListMenuContextData& ContextData, UToolMenu* Menu);
+	static UE_API void CreateSelectCollectionsSubMenu(UToolMenu* Menu, FObjectMixerEditorListMenuContextData ContextData);
+	static UE_API void AddCollectionsMenuItem(UToolMenu* InMenu, const FObjectMixerEditorListMenuContextData& ContextData);
+	static UE_API void OnTextCommitted(const FText& InText, ETextCommit::Type InCommitType, const FObjectMixerEditorListMenuContextData ContextData);
 	
-	static void OnCollectionMenuEntryCheckStateChanged(const FName Key, const FObjectMixerEditorListMenuContextData ContextData);
-	static void AddObjectsToCollection(const FName Key, const FObjectMixerEditorListMenuContextData ContextData);
-	static void RemoveObjectsFromCollection(const FName Key, const FObjectMixerEditorListMenuContextData ContextData);
-	static bool AreAllObjectsInCollection(const FName Key, const FObjectMixerEditorListMenuContextData ContextData);
-	static ECheckBoxState GetCheckStateForCollection(const FName Key, const FObjectMixerEditorListMenuContextData ContextData);
-
-	static FToolMenuEntry MakeCustomEditMenu(const FObjectMixerEditorListMenuContextData& ContextData);
-	static void ReplaceEditSubMenu(const FObjectMixerEditorListMenuContextData& ContextData);
+	static UE_API void OnCollectionMenuEntryCheckStateChanged(const FName Key, const FObjectMixerEditorListMenuContextData ContextData);
+	static UE_API void AddObjectsToCollection(const FName Key, const FObjectMixerEditorListMenuContextData ContextData);
+	static UE_API void RemoveObjectsFromCollection(const FName Key, const FObjectMixerEditorListMenuContextData ContextData);
+	static UE_API bool AreAllObjectsInCollection(const FName Key, const FObjectMixerEditorListMenuContextData ContextData);
+	static UE_API ECheckBoxState GetCheckStateForCollection(const FName Key, const FObjectMixerEditorListMenuContextData ContextData);
+
+	static UE_API FToolMenuEntry MakeCustomEditMenu(const FObjectMixerEditorListMenuContextData& ContextData);
+	static UE_API void ReplaceEditSubMenu(const FObjectMixerEditorListMenuContextData& ContextData);
 };
+
+#undef UE_API
diff --git a/Engine/Plugins/Editor/ObjectMixer/ObjectMixer/Source/ObjectMixer/Public/Views/Widgets/ObjectMixerEditorUWidget.h b/Engine/Plugins/Editor/ObjectMixer/ObjectMixer/Source/ObjectMixer/Public/Views/Widgets/ObjectMixerEditorUWidget.h
--- a/Engine/Plugins/Editor/ObjectMixer/ObjectMixer/Source/ObjectMixer/Public/Views/Widgets/ObjectMixerEditorUWidget.h
+++ b/Engine/Plugins/Editor/ObjectMixer/ObjectMixer/Source/ObjectMixer/Public/Views/Widgets/ObjectMixerEditorUWidget.h
@@ -8,6 +8,8 @@
 
 #include "ObjectMixerEditorUWidget.generated.h"
 
+#define UE_API OBJECTMIXEREDITOR_API
+
 // User Configurable Variables
 USTRUCT(BlueprintType)
 struct FObjectMixerWidgetUserConfig
@@ -25,21 +27,23 @@ struct FObjectMixerWidgetUserConfig
  * A UMG widget wrapper for the Object Mixer widget. Uses the same data as the Generic Object Mixer Instance.
  * Only useful in the editor. It is not compatible at runtime.
  */
-UCLASS(meta = (DisplayName="Object Mixer (Generic)"))
-class OBJECTMIXEREDITOR_API UObjectMixerEditorUWidget : public UWidget
+UCLASS(MinimalAPI, meta = (DisplayName="Object Mixer (Generic)"))
+class UObjectMixerEditorUWidget : public UWidget
 {
 	GENERATED_BODY()
 
 public:
 
 #if WITH_EDITOR
-	virtual const FText GetPaletteCategory() override;
+	UE_API virtual const FText GetPaletteCategory() override;
 #endif
 
 	// UWidget
-	virtual TSharedRef<SWidget> RebuildWidget() override;
+	UE_API virtual TSharedRef<SWidget> RebuildWidget() override;
 	// End of UWidget
 
 	UPROPERTY(EditAnywhere, BlueprintReadWrite, Category = "Object Mixer")
 	FObjectMixerWidgetUserConfig ObjectMixerWidgetUserConfig;
 };
+
+#undef UE_API
diff --git a/Engine/Plugins/Editor/ProxyLODPlugin/Source/ProxyLOD/Public/ProxyLODParameterization.h b/Engine/Plugins/Editor/ProxyLODPlugin/Source/ProxyLOD/Public/ProxyLODParameterization.h
--- a/Engine/Plugins/Editor/ProxyLODPlugin/Source/ProxyLOD/Public/ProxyLODParameterization.h
+++ b/Engine/Plugins/Editor/ProxyLODPlugin/Source/ProxyLOD/Public/ProxyLODParameterization.h
@@ -4,13 +4,15 @@
 
 #include "CoreMinimal.h"
 
+#define UE_API PROXYLODMESHREDUCTION_API
+
 struct FMeshDescription;
 
 
-class PROXYLODMESHREDUCTION_API IProxyLODParameterization
+class IProxyLODParameterization
 {
 public:
-	static TUniquePtr<IProxyLODParameterization> CreateTool();
+	static UE_API TUniquePtr<IProxyLODParameterization> CreateTool();
 	virtual ~IProxyLODParameterization(){}
 	
 	virtual bool ParameterizeMeshDescription(FMeshDescription& MeshDescription, 
@@ -60,3 +62,5 @@ public:
 		                     float&                   MaxStretch,
 		                     int32&                   NumCharts) const = 0;
 };
+
+#undef UE_API
diff --git a/Engine/Plugins/Editor/ProxyLODPlugin/Source/ProxyLOD/Public/ProxyLODVolume.h b/Engine/Plugins/Editor/ProxyLODPlugin/Source/ProxyLOD/Public/ProxyLODVolume.h
--- a/Engine/Plugins/Editor/ProxyLODPlugin/Source/ProxyLOD/Public/ProxyLODVolume.h
+++ b/Engine/Plugins/Editor/ProxyLODPlugin/Source/ProxyLOD/Public/ProxyLODVolume.h
@@ -5,15 +5,17 @@
 #include "CoreMinimal.h"
 #include "Math/Vector.h"
 
+#define UE_API PROXYLODMESHREDUCTION_API
+
 struct FMeshMergeData;
 struct FInstancedMeshMergeData;
 struct FMeshDescription;
 
-class PROXYLODMESHREDUCTION_API IProxyLODVolume
+class IProxyLODVolume
 {
 public:
 	/** Helper class to extract dimensions in voxel size units of OpenVDB volume */
-	struct PROXYLODMESHREDUCTION_API FVector3i
+	struct FVector3i
 	{
 		int32 X;
 		int32 Y;
@@ -29,12 +31,12 @@ public:
 			return Index == 0 ? X : (Index == 1 ? Y : (Index == 2 ? Z : 0));
 		}
 
-		int32 MinIndex() const;
+		UE_API int32 MinIndex() const;
 	};
 
 	/** Create OpenVDB volume from input geometry */
-	static TUniquePtr<IProxyLODVolume> CreateSDFVolumeFromMeshArray(const TArray<FMeshMergeData>& Geometry, float Step);
-	static TUniquePtr<IProxyLODVolume> CreateSDFVolumeFromMeshArray(const TArray<FInstancedMeshMergeData>& Geometry, float Step);
+	static UE_API TUniquePtr<IProxyLODVolume> CreateSDFVolumeFromMeshArray(const TArray<FMeshMergeData>& Geometry, float Step);
+	static UE_API TUniquePtr<IProxyLODVolume> CreateSDFVolumeFromMeshArray(const TArray<FInstancedMeshMergeData>& Geometry, float Step);
 
 	virtual ~IProxyLODVolume() {}
 
@@ -63,10 +65,10 @@ public:
 
 
 
-class PROXYLODMESHREDUCTION_API IVoxelBasedCSG
+class IVoxelBasedCSG
 {
 public :
-	static TUniquePtr<IVoxelBasedCSG> CreateCSGTool(float VoxelSize);
+	static UE_API TUniquePtr<IVoxelBasedCSG> CreateCSGTool(float VoxelSize);
 
 	virtual ~IVoxelBasedCSG() {}
 
@@ -126,4 +128,6 @@ public :
 	virtual bool ComputeDifference(IVoxelBasedCSG::FInterrupter& Interrupter, const FPlacedMesh& PlacedMeshA, const FPlacedMesh& PlacedMeshB, FMeshDescription& ResultMesh, FVector& AverageTranslation, double Adaptivity, double IsoSurface) const = 0;
 	virtual bool ComputeIntersection(IVoxelBasedCSG::FInterrupter& Interrupter, const FPlacedMesh& PlacedMeshA, const FPlacedMesh& PlacedMeshB, FMeshDescription& ResultMesh, FVector& AverageTranslation, double Adaptivity, double IsoSurface) const = 0;
 	virtual bool ComputeUnion(IVoxelBasedCSG::FInterrupter& Interrupter, const FPlacedMesh& PlacedMeshA, const FPlacedMesh& PlacedMeshB, FMeshDescription& ResultMesh, FVector& AverageTranslation, double Adaptivity, double IsoSurface) const = 0;
-};
\ No newline at end of file
+};
+
+#undef UE_API
diff --git a/Engine/Plugins/Editor/ScriptableToolsEditorMode/Source/ScriptableToolsEditorMode/Public/ScriptableToolsEditorMode.h b/Engine/Plugins/Editor/ScriptableToolsEditorMode/Source/ScriptableToolsEditorMode/Public/ScriptableToolsEditorMode.h
--- a/Engine/Plugins/Editor/ScriptableToolsEditorMode/Source/ScriptableToolsEditorMode/Public/ScriptableToolsEditorMode.h
+++ b/Engine/Plugins/Editor/ScriptableToolsEditorMode/Source/ScriptableToolsEditorMode/Public/ScriptableToolsEditorMode.h
@@ -11,6 +11,8 @@
 
 #include "ScriptableToolsEditorMode.generated.h"
 
+#define UE_API SCRIPTABLETOOLSEDITORMODE_API
+
 class FUICommandList;
 class FLevelObjectsObserver;
 class UModelingSceneSnappingManager;
@@ -20,45 +22,45 @@ class UScriptableToolSet;
 class UBlueprint;
 class UScriptableToolContextObject;
 
-UCLASS(Transient)
-class SCRIPTABLETOOLSEDITORMODE_API UScriptableToolsEditorMode : public UBaseLegacyWidgetEdMode
+UCLASS(MinimalAPI, Transient)
+class UScriptableToolsEditorMode : public UBaseLegacyWidgetEdMode
 {
 	GENERATED_BODY()
 public:
-	const static FEditorModeID EM_ScriptableToolsEditorModeId;
+	UE_API const static FEditorModeID EM_ScriptableToolsEditorModeId;
 
-	UScriptableToolsEditorMode();
-	UScriptableToolsEditorMode(FVTableHelper& Helper);
-	~UScriptableToolsEditorMode();
+	UE_API UScriptableToolsEditorMode();
+	UE_API UScriptableToolsEditorMode(FVTableHelper& Helper);
+	UE_API ~UScriptableToolsEditorMode();
 	////////////////
 	// UEdMode interface
 	////////////////
 
-	virtual void Tick(FEditorViewportClient* ViewportClient, float DeltaTime) override;
+	UE_API virtual void Tick(FEditorViewportClient* ViewportClient, float DeltaTime) override;
 
-	virtual void ActorSelectionChangeNotify() override;
+	UE_API virtual void ActorSelectionChangeNotify() override;
 
-	virtual bool ShouldDrawWidget() const override;
-	virtual bool ProcessEditDelete() override;
-	virtual bool ProcessEditCut() override;
+	UE_API virtual bool ShouldDrawWidget() const override;
+	UE_API virtual bool ProcessEditDelete() override;
+	UE_API virtual bool ProcessEditCut() override;
 
-	virtual bool CanAutoSave() const override;
+	UE_API virtual bool CanAutoSave() const override;
 
-	virtual bool ComputeBoundingBoxForViewportFocus(AActor* Actor, UPrimitiveComponent* PrimitiveComponent, FBox& InOutBox) const override;
+	UE_API virtual bool ComputeBoundingBoxForViewportFocus(AActor* Actor, UPrimitiveComponent* PrimitiveComponent, FBox& InOutBox) const override;
 
-	virtual bool GetPivotForOrbit(FVector& OutPivot) const override;
+	UE_API virtual bool GetPivotForOrbit(FVector& OutPivot) const override;
 
 	/*
 	 * focus events
 	 */
 
 	// called when we "start" this editor mode (ie switch to this tab)
-	virtual void Enter() override;
+	UE_API virtual void Enter() override;
 
 	// called when we "end" this editor mode (ie switch to another tab)
-	virtual void Exit() override;
+	UE_API virtual void Exit() override;
 
-	virtual bool ShouldToolStartBeAllowed(const FString& ToolIdentifier) const override;
+	UE_API virtual bool ShouldToolStartBeAllowed(const FString& ToolIdentifier) const override;
 
 	//////////////////
 	// End of UEdMode interface
@@ -67,21 +69,21 @@ public:
 	
 
 protected:
-	virtual void BindCommands() override;
-	virtual void CreateToolkit() override;
-	virtual void OnToolStarted(UInteractiveToolManager* Manager, UInteractiveTool* Tool) override;
-	virtual void OnToolEnded(UInteractiveToolManager* Manager, UInteractiveTool* Tool) override;
+	UE_API virtual void BindCommands() override;
+	UE_API virtual void CreateToolkit() override;
+	UE_API virtual void OnToolStarted(UInteractiveToolManager* Manager, UInteractiveTool* Tool) override;
+	UE_API virtual void OnToolEnded(UInteractiveToolManager* Manager, UInteractiveTool* Tool) override;
 	
-	virtual void OnToolPostBuild(UInteractiveToolManager* InToolManager, EToolSide InSide, UInteractiveTool* InBuiltTool, UInteractiveToolBuilder* InToolBuilder, const FToolBuilderState& ToolState);
+	UE_API virtual void OnToolPostBuild(UInteractiveToolManager* InToolManager, EToolSide InSide, UInteractiveTool* InBuiltTool, UInteractiveToolBuilder* InToolBuilder, const FToolBuilderState& ToolState);
 
-	void OnToolsContextRender(IToolsContextRenderAPI* RenderAPI);
+	UE_API void OnToolsContextRender(IToolsContextRenderAPI* RenderAPI);
 
-	void AcceptActiveToolActionOrTool();
-	void CancelActiveToolActionOrTool();
+	UE_API void AcceptActiveToolActionOrTool();
+	UE_API void CancelActiveToolActionOrTool();
 
-	void ConfigureRealTimeViewportsOverride(bool bEnable);
+	UE_API void ConfigureRealTimeViewportsOverride(bool bEnable);
 
-	void OnBlueprintPreCompile(UBlueprint* Blueprint);
+	UE_API void OnBlueprintPreCompile(UBlueprint* Blueprint);
 	FDelegateHandle BlueprintPreCompileHandle;
 
 	UE_DEPRECATED(5.6, "Deprecated in favor of OnBlueprintPreCompile")
@@ -90,9 +92,9 @@ protected:
 	UE_DEPRECATED(5.6, "Deprecated in favor of OnBlueprintPreCompile")
 	FDelegateHandle BlueprintCompiledHandle;
 
-	void InitializeModeContexts();
+	UE_API void InitializeModeContexts();
 
-	void RebuildScriptableToolSet();
+	UE_API void RebuildScriptableToolSet();
 
 private:
 
@@ -108,3 +110,5 @@ public:
 	virtual UScriptableToolSet* GetActiveScriptableTools() { return ScriptableTools; }
 
 };
+
+#undef UE_API
diff --git a/Engine/Plugins/Editor/ScriptableToolsEditorMode/Source/ScriptableToolsEditorMode/Public/ScriptableToolsEditorModeSettings.h b/Engine/Plugins/Editor/ScriptableToolsEditorMode/Source/ScriptableToolsEditorMode/Public/ScriptableToolsEditorModeSettings.h
--- a/Engine/Plugins/Editor/ScriptableToolsEditorMode/Source/ScriptableToolsEditorMode/Public/ScriptableToolsEditorModeSettings.h
+++ b/Engine/Plugins/Editor/ScriptableToolsEditorMode/Source/ScriptableToolsEditorMode/Public/ScriptableToolsEditorModeSettings.h
@@ -8,9 +8,11 @@
 
 #include "ScriptableToolsEditorModeSettings.generated.h"
 
+#define UE_API SCRIPTABLETOOLSEDITORMODE_API
 
-UCLASS(config = Editor)
-class SCRIPTABLETOOLSEDITORMODE_API UScriptableToolsModeCustomizationSettings : public UDeveloperSettings
+
+UCLASS(MinimalAPI, config = Editor)
+class UScriptableToolsModeCustomizationSettings : public UDeveloperSettings
 {
 	GENERATED_BODY()
 
@@ -21,8 +23,8 @@ public:
 	virtual FName GetCategoryName() const { return FName("Plugins"); }
 	virtual FName GetSectionName() const { return FName("ScriptableTools"); }
 
-	virtual FText GetSectionText() const override;
-	virtual FText GetSectionDescription() const override;
+	UE_API virtual FText GetSectionText() const override;
+	UE_API virtual FText GetSectionDescription() const override;
 
 public:
 
@@ -39,4 +41,6 @@ public:
 
 	bool RegisterAllTools() const {	return ToolRegistrationFilters.GetGroups().IsEmpty(); }
 
-};
\ No newline at end of file
+};
+
+#undef UE_API
diff --git a/Engine/Plugins/Editor/ScriptableToolsEditorMode/Source/ScriptableToolsEditorMode/Public/ScriptableToolsEditorModeStyle.h b/Engine/Plugins/Editor/ScriptableToolsEditorMode/Source/ScriptableToolsEditorMode/Public/ScriptableToolsEditorModeStyle.h
--- a/Engine/Plugins/Editor/ScriptableToolsEditorMode/Source/ScriptableToolsEditorMode/Public/ScriptableToolsEditorModeStyle.h
+++ b/Engine/Plugins/Editor/ScriptableToolsEditorMode/Source/ScriptableToolsEditorMode/Public/ScriptableToolsEditorModeStyle.h
@@ -6,19 +6,21 @@
 #include "Styling/ISlateStyle.h"
 #include "Styling/SlateStyle.h"
 
-class SCRIPTABLETOOLSEDITORMODE_API FScriptableToolsEditorModeStyle
+#define UE_API SCRIPTABLETOOLSEDITORMODE_API
+
+class FScriptableToolsEditorModeStyle
 {
 public:
-	static void Initialize();
+	static UE_API void Initialize();
 
-	static void Shutdown();
+	static UE_API void Shutdown();
 
-	static TSharedPtr< class ISlateStyle > Get();
+	static UE_API TSharedPtr< class ISlateStyle > Get();
 
-	static FName GetStyleSetName();
+	static UE_API FName GetStyleSetName();
 
 	// use to access icons defined by the style set by name, eg GetBrush("BrushFalloffIcons.Smooth")
-	static const FSlateBrush* GetBrush(FName PropertyName, const ANSICHAR* Specifier = NULL);
+	static UE_API const FSlateBrush* GetBrush(FName PropertyName, const ANSICHAR* Specifier = NULL);
 
 
 	// Try to find a png or svg icon at FileNameWithPath and register it in the style set with the given StyleIconIdentifier,
@@ -31,12 +33,14 @@ public:
 	// the user enters as (eg) a property/etc. To help with debugging, this path can be provided as ExternalRelativePath,
 	// and it will be printed out instead of the absolute path as part of the log error message.
 	//
-	static bool TryRegisterCustomIcon(FName StyleIconIdentifier, FString FileNameWithPath, FString ExternalRelativePath);
+	static UE_API bool TryRegisterCustomIcon(FName StyleIconIdentifier, FString FileNameWithPath, FString ExternalRelativePath);
 
 
 private:
-	static FString InContent(const FString& RelativePath, const ANSICHAR* Extension);
+	static UE_API FString InContent(const FString& RelativePath, const ANSICHAR* Extension);
 
 private:
-	static TSharedPtr< class FSlateStyleSet > StyleSet;
+	static UE_API TSharedPtr< class FSlateStyleSet > StyleSet;
 };
+
+#undef UE_API
diff --git a/Engine/Plugins/Editor/ScriptableToolsEditorMode/Source/ScriptableToolsEditorMode/Public/ScriptableToolsEditorModeToolkit.h b/Engine/Plugins/Editor/ScriptableToolsEditorMode/Source/ScriptableToolsEditorMode/Public/ScriptableToolsEditorModeToolkit.h
--- a/Engine/Plugins/Editor/ScriptableToolsEditorMode/Source/ScriptableToolsEditorMode/Public/ScriptableToolsEditorModeToolkit.h
+++ b/Engine/Plugins/Editor/ScriptableToolsEditorMode/Source/ScriptableToolsEditorMode/Public/ScriptableToolsEditorModeToolkit.h
@@ -9,6 +9,8 @@
 #include "Widgets/Input/SEditableTextBox.h"
 #include "StatusBarSubsystem.h"
 
+#define UE_API SCRIPTABLETOOLSEDITORMODE_API
+
 
 class IDetailsView;
 class SButton;
@@ -16,77 +18,77 @@ class STextBlock;
 class UBlueprint;
 class UBaseScriptableToolBuilder;
 
-class SCRIPTABLETOOLSEDITORMODE_API FScriptableToolsEditorModeToolkit : public FModeToolkit
+class FScriptableToolsEditorModeToolkit : public FModeToolkit
 {
 public:
 
-	FScriptableToolsEditorModeToolkit();
-	~FScriptableToolsEditorModeToolkit();
+	UE_API FScriptableToolsEditorModeToolkit();
+	UE_API ~FScriptableToolsEditorModeToolkit();
 	
 	/** FModeToolkit interface */
-	virtual void Init(const TSharedPtr<IToolkitHost>& InitToolkitHost, TWeakObjectPtr<UEdMode> InOwningMode) override;
+	UE_API virtual void Init(const TSharedPtr<IToolkitHost>& InitToolkitHost, TWeakObjectPtr<UEdMode> InOwningMode) override;
 
 	/** IToolkit interface */
-	virtual FName GetToolkitFName() const override;
-	virtual FText GetBaseToolkitName() const override;
+	UE_API virtual FName GetToolkitFName() const override;
+	UE_API virtual FText GetBaseToolkitName() const override;
 	virtual TSharedPtr<class SWidget> GetInlineContent() const override { return ToolkitWidget; }
 
 	// initialize toolkit widgets that need to wait until mode is initialized/entered
-	virtual void InitializeAfterModeSetup();
+	UE_API virtual void InitializeAfterModeSetup();
 
 	// set/clear notification message area
-	virtual void PostNotification(const FText& Message);
-	virtual void ClearNotification();
+	UE_API virtual void PostNotification(const FText& Message);
+	UE_API virtual void ClearNotification();
 
 	// set/clear warning message area
-	virtual void PostWarning(const FText& Message);
-	virtual void ClearWarning();
+	UE_API virtual void PostWarning(const FText& Message);
+	UE_API virtual void ClearWarning();
 
 	// Async Tool Loading
-	void StartAsyncToolLoading();
-	void SetAsyncProgress(float PercentLoaded);
-	void EndAsyncToolLoading();
-	bool AreToolsLoading() const;
-	TOptional<float> GetToolPercentLoaded() const;
+	UE_API void StartAsyncToolLoading();
+	UE_API void SetAsyncProgress(float PercentLoaded);
+	UE_API void EndAsyncToolLoading();
+	UE_API bool AreToolsLoading() const;
+	UE_API TOptional<float> GetToolPercentLoaded() const;
 
 	/** Returns the Mode specific tabs in the mode toolbar **/ 
-	virtual void GetToolPaletteNames(TArray<FName>& InPaletteName) const override;
-	virtual FText GetToolPaletteDisplayName(FName PaletteName) const override; 
-	virtual void BuildToolPalette(FName PaletteName, class FToolBarBuilder& ToolbarBuilder) override;
-	virtual void OnToolPaletteChanged(FName PaletteName) override;
+	UE_API virtual void GetToolPaletteNames(TArray<FName>& InPaletteName) const override;
+	UE_API virtual FText GetToolPaletteDisplayName(FName PaletteName) const override; 
+	UE_API virtual void BuildToolPalette(FName PaletteName, class FToolBarBuilder& ToolbarBuilder) override;
+	UE_API virtual void OnToolPaletteChanged(FName PaletteName) override;
 	virtual bool HasIntegratedToolPalettes() const override { return false; }
 	virtual bool HasExclusiveToolPalettes() const override { return false; }
 
 	virtual FText GetActiveToolDisplayName() const override { return ActiveToolName; }
 	virtual FText GetActiveToolMessage() const override { return ActiveToolMessage; }
 
-	virtual void EnableShowRealtimeWarning(bool bEnable);
+	UE_API virtual void EnableShowRealtimeWarning(bool bEnable);
 
-	virtual void OnToolStarted(UInteractiveToolManager* Manager, UInteractiveTool* Tool) override;
-	virtual void OnToolEnded(UInteractiveToolManager* Manager, UInteractiveTool* Tool) override;
+	UE_API virtual void OnToolStarted(UInteractiveToolManager* Manager, UInteractiveTool* Tool) override;
+	UE_API virtual void OnToolEnded(UInteractiveToolManager* Manager, UInteractiveTool* Tool) override;
 
-	virtual void CustomizeModeDetailsViewArgs(FDetailsViewArgs& ArgsInOut) override;
+	UE_API virtual void CustomizeModeDetailsViewArgs(FDetailsViewArgs& ArgsInOut) override;
 
-	void OnActiveViewportChanged(TSharedPtr<IAssetViewport>, TSharedPtr<IAssetViewport> );
+	UE_API void OnActiveViewportChanged(TSharedPtr<IAssetViewport>, TSharedPtr<IAssetViewport> );
 
-	virtual void InvokeUI() override;
+	UE_API virtual void InvokeUI() override;
 
-	virtual void ForceToolPaletteRebuild();
+	UE_API virtual void ForceToolPaletteRebuild();
 
-	void GetActiveToolPaletteNames(TArray<FName>& OutPaletteNames);
+	UE_API void GetActiveToolPaletteNames(TArray<FName>& OutPaletteNames);
 
 protected:
 
 	/** FModeToolkit interface */
-	virtual void RebuildModeToolBar() override;
-	virtual bool ShouldShowModeToolbar() const override;
+	UE_API virtual void RebuildModeToolBar() override;
+	UE_API virtual bool ShouldShowModeToolbar() const override;
 
-	void RebuildModeToolPaletteWidgets();
-	void RebuildModeToolkitBuilderPalettes();
+	UE_API void RebuildModeToolPaletteWidgets();
+	UE_API void RebuildModeToolkitBuilderPalettes();
 
 
 private:
-	const static TArray<FName> PaletteNames_Standard;
+	UE_API const static TArray<FName> PaletteNames_Standard;
 
 	FText ActiveToolName;
 	FText ActiveToolMessage;
@@ -94,10 +96,10 @@ private:
 	const FSlateBrush* ActiveToolIcon = nullptr;
 
 	TSharedPtr<SWidget> ToolkitWidget;
-	void UpdateActiveToolProperties();
-	void InvalidateCachedDetailPanelState(UObject* ChangedObject);
+	UE_API void UpdateActiveToolProperties();
+	UE_API void InvalidateCachedDetailPanelState(UObject* ChangedObject);
 	
-	void RegisterPalettes();
+	UE_API void RegisterPalettes();
 	FDelegateHandle ActivePaletteChangedHandle;
 
 	TSharedPtr<SWidget> ViewportOverlayWidget;
@@ -121,7 +123,7 @@ private:
 	TSharedPtr<SWidget> ToolPaletteLoadBar;
 
 	bool bShowRealtimeWarning = false;
-	void UpdateShowWarnings();
+	UE_API void UpdateShowWarnings();
 
 	struct FScriptableToolData
 	{
@@ -130,7 +132,7 @@ private:
 		UBaseScriptableToolBuilder* Builder = nullptr;
 	};
 	TMap<FName, TArray<FScriptableToolData>> ActiveToolCategories;
-	void UpdateActiveToolCategories();
+	UE_API void UpdateActiveToolCategories();
 
 	bool bFirstInitializeAfterModeSetup = true;
 
@@ -139,3 +141,5 @@ private:
 
 	// custom accept/cancel/complete handlers
 };
+
+#undef UE_API
diff --git a/Engine/Plugins/Editor/SequencerAnimTools/Source/SequencerAnimTools/Public/BaseSequencerAnimTool.h b/Engine/Plugins/Editor/SequencerAnimTools/Source/SequencerAnimTools/Public/BaseSequencerAnimTool.h
--- a/Engine/Plugins/Editor/SequencerAnimTools/Source/SequencerAnimTools/Public/BaseSequencerAnimTool.h
+++ b/Engine/Plugins/Editor/SequencerAnimTools/Source/SequencerAnimTools/Public/BaseSequencerAnimTool.h
@@ -6,6 +6,8 @@
 #include "UObject/Interface.h"
 #include "BaseSequencerAnimTool.generated.h"
 
+#define UE_API SEQUENCERANIMTOOLS_API
+
 class UInteractiveToolManager;
 class UTransformProxy;
 class UCombinedTransformGizmo;
@@ -18,14 +20,14 @@ class UBaseSequencerAnimTool : public UInterface
 	GENERATED_UINTERFACE_BODY()
 };
 
-class SEQUENCERANIMTOOLS_API IBaseSequencerAnimTool
+class IBaseSequencerAnimTool
 {
 	GENERATED_IINTERFACE_BODY()
 	virtual bool ProcessCommandBindings(const FKey Key, const bool bRepeat) const { return false; }
 
 };
 
-struct SEQUENCERANIMTOOLS_API FSequencerAnimToolHelpers
+struct FSequencerAnimToolHelpers
 {
 	struct FGizmoData
 	{
@@ -35,5 +37,7 @@ struct SEQUENCERANIMTOOLS_API FSequencerAnimToolHelpers
 		UInteractiveGizmoManager* GizmoManager= nullptr;
 		FString InstanceIdentifier;
 	};
-	static void CreateGizmo(const FGizmoData& InData, TObjectPtr <UCombinedTransformGizmo>& OutCombinedGizmo, TObjectPtr<UTransformGizmo>& OutTRSGizmo);
+	static UE_API void CreateGizmo(const FGizmoData& InData, TObjectPtr <UCombinedTransformGizmo>& OutCombinedGizmo, TObjectPtr<UTransformGizmo>& OutTRSGizmo);
 };
+
+#undef UE_API
diff --git a/Engine/Plugins/Editor/UMGWidgetPreview/Source/UMGWidgetPreview/Public/WidgetPreview.h b/Engine/Plugins/Editor/UMGWidgetPreview/Source/UMGWidgetPreview/Public/WidgetPreview.h
--- a/Engine/Plugins/Editor/UMGWidgetPreview/Source/UMGWidgetPreview/Public/WidgetPreview.h
+++ b/Engine/Plugins/Editor/UMGWidgetPreview/Source/UMGWidgetPreview/Public/WidgetPreview.h
@@ -6,6 +6,8 @@
 
 #include "WidgetPreview.generated.h"
 
+#define UE_API UMGWIDGETPREVIEW_API
+
 class SWidget;
 class UPanelWidget;
 class UUserWidget;
@@ -25,7 +27,7 @@ enum class EWidgetPreviewWidgetChangeType : uint8
 };
 
 USTRUCT(BlueprintType)
-struct UMGWIDGETPREVIEW_API FPreviewableWidgetVariant
+struct FPreviewableWidgetVariant
 {
 	GENERATED_BODY();
 
@@ -33,18 +35,18 @@ struct UMGWIDGETPREVIEW_API FPreviewableWidgetVariant
 	FSoftObjectPath ObjectPath;
 
 	FPreviewableWidgetVariant() = default;
-	explicit FPreviewableWidgetVariant(const TSubclassOf<UUserWidget>& InWidgetType);
-	explicit FPreviewableWidgetVariant(const UWidgetPreview* InWidgetPreview);
+	UE_API explicit FPreviewableWidgetVariant(const TSubclassOf<UUserWidget>& InWidgetType);
+	UE_API explicit FPreviewableWidgetVariant(const UWidgetPreview* InWidgetPreview);
 
 public:
 	/** Flushes cached widgets and re-resolves from the ObjectPath. */
-	void UpdateCachedWidget();
+	UE_API void UpdateCachedWidget();
 
 	/** Returns the referenced Object as a UUserWidget (CDO). Returns nullptr if not found, or we couldn't find a nested UUserWidget (ie. inside a UWidgetPreview). */
-	const UUserWidget* AsUserWidgetCDO() const;
+	UE_API const UUserWidget* AsUserWidgetCDO() const;
 
 	/** Returns the referenced Object as a UWidgetPreview. Returns nullptr if not found, or not a UWidgetPreview. */
-	const UWidgetPreview* AsWidgetPreview() const;
+	UE_API const UWidgetPreview* AsWidgetPreview() const;
 
 	friend bool operator==(const FPreviewableWidgetVariant& Left, const FPreviewableWidgetVariant& Right)
 	{
@@ -64,24 +66,24 @@ private:
 	TWeakObjectPtr<UWidgetPreview> CachedWidgetPreview;
 };
 
-UCLASS(BlueprintType, NotBlueprintable, AutoExpandCategories = "Widgets")
-class UMGWIDGETPREVIEW_API UWidgetPreview
+UCLASS(MinimalAPI, BlueprintType, NotBlueprintable, AutoExpandCategories = "Widgets")
+class UWidgetPreview
 	: public UObject
 {
 	GENERATED_BODY()
 
 public:
 
-	UWidgetPreview(const FObjectInitializer& ObjectInitializer);
+	UE_API UWidgetPreview(const FObjectInitializer& ObjectInitializer);
 
-	virtual void BeginDestroy() override;
+	UE_API virtual void BeginDestroy() override;
 
 	/** Convenience function to check that all utilized widgets have bCanCallInitializedWithoutPlayerContext set to true, and reports any that don't. */
-	bool CanCallInitializedWithoutPlayerContext(const bool bInRecursive, TArray<const UUserWidget*>& OutFailedWidgets);
+	UE_API bool CanCallInitializedWithoutPlayerContext(const bool bInRecursive, TArray<const UUserWidget*>& OutFailedWidgets);
 
 	// @todo: move to utility func somewhere else?
 	/** Convenience function to check that the provided widget (and it's children) has bCanCallInitializedWithoutPlayerContext set to true, and reports any that don't. */
-	static bool CanCallInitializedWithoutPlayerContextOnWidget(const UUserWidget* InUserWidget, const bool bInRecursive, TArray<const UUserWidget*>& OutFailedWidgets);
+	static UE_API bool CanCallInitializedWithoutPlayerContextOnWidget(const UUserWidget* InUserWidget, const bool bInRecursive, TArray<const UUserWidget*>& OutFailedWidgets);
 
 public:
 	using FOnWidgetChanged = TMulticastDelegate<void(const EWidgetPreviewWidgetChangeType)>;
@@ -89,53 +91,53 @@ public:
 	FOnWidgetChanged& OnWidgetChanged() { return OnWidgetChangedDelegate; }
 
 	UFUNCTION(BlueprintCallable, Category = "Layout")
-	const TArray<FName>& GetWidgetSlotNames() const;
+	UE_API const TArray<FName>& GetWidgetSlotNames() const;
 
 	/** Returns or builds and returns an instance of the root widget for previewing. Can be used to trigger a rebuild. */
-	[[maybe_unused]] UUserWidget* GetOrCreateWidgetInstance(UWorld* InWorld, const bool bInForceRecreate = false);
+	[[maybe_unused]] UE_API UUserWidget* GetOrCreateWidgetInstance(UWorld* InWorld, const bool bInForceRecreate = false);
 
 	/** Returns the current widget instance, if any. */
-	UUserWidget* GetWidgetInstance() const;
+	UE_API UUserWidget* GetWidgetInstance() const;
 
 	/** Returns the current underlying slate widget instance, if any. */
-	TSharedPtr<SWidget> GetSlateWidgetInstance() const;
+	UE_API TSharedPtr<SWidget> GetSlateWidgetInstance() const;
 
 	/** Stores the current instance in PreviousWidgetInstance, and clears WidgetInstance. */
-	void ClearWidgetInstance();
+	UE_API void ClearWidgetInstance();
 
-	const UUserWidget* GetWidgetCDO() const;
-	const UUserWidget* GetWidgetCDOForSlot(const FName InSlotName) const;
+	UE_API const UUserWidget* GetWidgetCDO() const;
+	UE_API const UUserWidget* GetWidgetCDOForSlot(const FName InSlotName) const;
 
-	const FPreviewableWidgetVariant& GetWidgetType() const;
-	void SetWidgetType(const FPreviewableWidgetVariant& InWidget);
+	UE_API const FPreviewableWidgetVariant& GetWidgetType() const;
+	UE_API void SetWidgetType(const FPreviewableWidgetVariant& InWidget);
 
-	const TMap<FName, FPreviewableWidgetVariant>& GetSlotWidgetTypes() const;
-	void SetSlotWidgetTypes(const TMap<FName, FPreviewableWidgetVariant>& InWidgets);
+	UE_API const TMap<FName, FPreviewableWidgetVariant>& GetSlotWidgetTypes() const;
+	UE_API void SetSlotWidgetTypes(const TMap<FName, FPreviewableWidgetVariant>& InWidgets);
 
-	const bool GetbShouldOverrideWidgetSize() const;
-	void SetbShouldOverrideWidgetSize(bool InOverride);
+	UE_API const bool GetbShouldOverrideWidgetSize() const;
+	UE_API void SetbShouldOverrideWidgetSize(bool InOverride);
 
-	const FVector2D GetOverriddenWidgetSize() const;
-	void SetOverriddenWidgetSize(FVector2D InWidgetSize);
+	UE_API const FVector2D GetOverriddenWidgetSize() const;
+	UE_API void SetOverriddenWidgetSize(FVector2D InWidgetSize);
 
 protected:
-	virtual void PostLoad() override;
+	UE_API virtual void PostLoad() override;
 
-	virtual void PostEditChangeProperty(FPropertyChangedEvent& PropertyChangedEvent) override;
+	UE_API virtual void PostEditChangeProperty(FPropertyChangedEvent& PropertyChangedEvent) override;
 
-	void OnWidgetBlueprintChanged(UBlueprint* InBlueprint);
+	UE_API void OnWidgetBlueprintChanged(UBlueprint* InBlueprint);
 
 	/** Misc. functionality to perform after a widget assignment is changed. */
-	void UpdateWidgets();
+	UE_API void UpdateWidgets();
 
 	/** Creates a new WidgetInstance, replacing the current one if it exists. */
-	UUserWidget* CreateWidgetInstance(UWorld* InWorld);
+	UE_API UUserWidget* CreateWidgetInstance(UWorld* InWorld);
 
-	void CleanupReferences();
+	UE_API void CleanupReferences();
 
 	/** Returns slot names not already occupied in SlotWidgets. */
 	UFUNCTION(BlueprintCallable, Category = "Widget")
-	TArray<FName> GetAvailableWidgetSlotNames();
+	UE_API TArray<FName> GetAvailableWidgetSlotNames();
 
 private:
 	/** Widget to preview. */
@@ -168,3 +170,5 @@ private:
 
 	FOnWidgetChanged OnWidgetChangedDelegate;
 };
+
+#undef UE_API
diff --git a/Engine/Plugins/Editor/UVEditor/Source/UVEditor/Private/UVEditorUXPropertySets.h b/Engine/Plugins/Editor/UVEditor/Source/UVEditor/Private/UVEditorUXPropertySets.h
--- a/Engine/Plugins/Editor/UVEditor/Source/UVEditor/Private/UVEditorUXPropertySets.h
+++ b/Engine/Plugins/Editor/UVEditor/Source/UVEditor/Private/UVEditorUXPropertySets.h
@@ -9,8 +9,8 @@
 /**
 * Visualization settings for the UUVEditorMode's Unwrapped Visualization
  */
-UCLASS()
-class UVEDITOR_API UUVEditorUnwrappedUXProperties : public UInteractiveToolPropertySet
+UCLASS(MinimalAPI)
+class UUVEditorUnwrappedUXProperties : public UInteractiveToolPropertySet
 {
 	GENERATED_BODY()
 public:
@@ -28,8 +28,8 @@ public:
 /**
  * Visualization settings for the UUVEditorMode's properties in the Live Preview (3d) Viewport
  */
-UCLASS()
-class UVEDITOR_API UUVEditorLivePreviewUXProperties : public UInteractiveToolPropertySet
+UCLASS(MinimalAPI)
+class UUVEditorLivePreviewUXProperties : public UInteractiveToolPropertySet
 {
 	GENERATED_BODY()
 public:
@@ -41,4 +41,4 @@ public:
 
 	UPROPERTY(EditAnywhere, Category="Live Preview Visualization", meta = (UIMin = "0.0", UIMax = "10.0", ClampMin = "0.0", DisplayName = "Selection Point Size"))
 	float SelectionPointSize;
-};
\ No newline at end of file
+};
diff --git a/Engine/Plugins/Editor/UVEditor/Source/UVEditor/Public/UVEditor.h b/Engine/Plugins/Editor/UVEditor/Source/UVEditor/Public/UVEditor.h
--- a/Engine/Plugins/Editor/UVEditor/Source/UVEditor/Public/UVEditor.h
+++ b/Engine/Plugins/Editor/UVEditor/Source/UVEditor/Public/UVEditor.h
@@ -8,6 +8,8 @@
 
 #include "UVEditor.generated.h"
 
+#define UE_API UVEDITOR_API
+
 
 class FUICommandList;
 class UEdModeInteractiveToolsContext;
@@ -31,23 +33,23 @@ class IAssetEditorInstance;
  * the asset editor toolkit, though the mode toolkit has most of the things that deal with the toolbar
  * on the left.
  */
-UCLASS()
-class UVEDITOR_API UUVEditor : public UAssetEditor
+UCLASS(MinimalAPI)
+class UUVEditor : public UAssetEditor
 {
 	GENERATED_BODY()
 
 public:
-	void Initialize(const TArray<TObjectPtr<UObject>>& InObjects);
-	void Initialize(const TArray<TObjectPtr<UObject>>& InObjects, const TArray<FTransform3d>& InTransforms);
+	UE_API void Initialize(const TArray<TObjectPtr<UObject>>& InObjects);
+	UE_API void Initialize(const TArray<TObjectPtr<UObject>>& InObjects, const TArray<FTransform3d>& InTransforms);
 
-	void GetWorldspaceRelativeTransforms(TArray<FTransform3d>& OutTransforms);
+	UE_API void GetWorldspaceRelativeTransforms(TArray<FTransform3d>& OutTransforms);
 
 	/** Returns the asset editor instance interface, so that its window can be focused, for example. */
-	IAssetEditorInstance* GetInstanceInterface();
+	UE_API IAssetEditorInstance* GetInstanceInterface();
 
 	// UAssetEditor overrides
-	void GetObjectsToEdit(TArray<UObject*>& OutObjects) override;
-	virtual TSharedPtr<FBaseAssetToolkit> CreateToolkit() override;
+	UE_API void GetObjectsToEdit(TArray<UObject*>& OutObjects) override;
+	UE_API virtual TSharedPtr<FBaseAssetToolkit> CreateToolkit() override;
 
 	// Note that things like CloseWindow, FocusWindow, and IsPrimaryEditor seem to be called
 	// on the toolkit, not the editor class, so they are overriden there.
@@ -62,3 +64,5 @@ protected:
 
 	friend class UUVEditorSubsystem;
 };
+
+#undef UE_API
diff --git a/Engine/Plugins/Editor/UVEditor/Source/UVEditor/Public/UVEditor2DViewportClient.h b/Engine/Plugins/Editor/UVEditor/Source/UVEditor/Public/UVEditor2DViewportClient.h
--- a/Engine/Plugins/Editor/UVEditor/Source/UVEditor/Public/UVEditor2DViewportClient.h
+++ b/Engine/Plugins/Editor/UVEditor/Source/UVEditor/Public/UVEditor2DViewportClient.h
@@ -9,6 +9,8 @@
 #include "Behaviors/2DViewportBehaviorTargets.h" // FEditor2DScrollBehaviorTarget, FEditor2DMouseWheelZoomBehaviorTarget
 #include "ContextObjects/UVToolViewportButtonsAPI.h" // UUVToolViewportButtonsAPI::ESelectionMode
 
+#define UE_API UVEDITOR_API
+
 class UCanvas;
 class UUVTool2DViewportAPI;
 
@@ -16,56 +18,56 @@ class UUVTool2DViewportAPI;
  * Client used to display a 2D view of the UV's, implemented by using a perspective viewport with a locked
  * camera.
  */
-class UVEDITOR_API FUVEditor2DViewportClient : public FEditorViewportClient, public IInputBehaviorSource
+class FUVEditor2DViewportClient : public FEditorViewportClient, public IInputBehaviorSource
 {
 public:
-	FUVEditor2DViewportClient(FEditorModeTools* InModeTools, FPreviewScene* InPreviewScene,
+	UE_API FUVEditor2DViewportClient(FEditorModeTools* InModeTools, FPreviewScene* InPreviewScene,
 		const TWeakPtr<SEditorViewport>& InEditorViewportWidget, UUVToolViewportButtonsAPI* ViewportButtonsAPI, UUVTool2DViewportAPI* UVTool2DViewportAPI);
 
 	virtual ~FUVEditor2DViewportClient() {}
 
-	bool AreSelectionButtonsEnabled() const;
-	void SetSelectionMode(UUVToolViewportButtonsAPI::ESelectionMode NewMode);
-	UUVToolViewportButtonsAPI::ESelectionMode GetSelectionMode() const;
-	bool AreWidgetButtonsEnabled() const;
-
-	void SetLocationGridSnapEnabled(bool bEnabled);
-	bool GetLocationGridSnapEnabled();
-	void SetLocationGridSnapValue(float SnapValue);
-	float GetLocationGridSnapValue();
-	void SetRotationGridSnapEnabled(bool bEnabled);
-	bool GetRotationGridSnapEnabled();
-	void SetRotationGridSnapValue(float SnapValue);
-	float GetRotationGridSnapValue();
-	void SetScaleGridSnapEnabled(bool bEnabled);
-	bool GetScaleGridSnapEnabled();
-	void SetScaleGridSnapValue(float SnapValue);
-	float GetScaleGridSnapValue();
+	UE_API bool AreSelectionButtonsEnabled() const;
+	UE_API void SetSelectionMode(UUVToolViewportButtonsAPI::ESelectionMode NewMode);
+	UE_API UUVToolViewportButtonsAPI::ESelectionMode GetSelectionMode() const;
+	UE_API bool AreWidgetButtonsEnabled() const;
+
+	UE_API void SetLocationGridSnapEnabled(bool bEnabled);
+	UE_API bool GetLocationGridSnapEnabled();
+	UE_API void SetLocationGridSnapValue(float SnapValue);
+	UE_API float GetLocationGridSnapValue();
+	UE_API void SetRotationGridSnapEnabled(bool bEnabled);
+	UE_API bool GetRotationGridSnapEnabled();
+	UE_API void SetRotationGridSnapValue(float SnapValue);
+	UE_API float GetRotationGridSnapValue();
+	UE_API void SetScaleGridSnapEnabled(bool bEnabled);
+	UE_API bool GetScaleGridSnapEnabled();
+	UE_API void SetScaleGridSnapValue(float SnapValue);
+	UE_API float GetScaleGridSnapValue();
 
 	// FEditorViewportClient
-	virtual bool InputKey(const FInputKeyEventArgs& EventArgs) override;
+	UE_API virtual bool InputKey(const FInputKeyEventArgs& EventArgs) override;
 
-	virtual void Draw(const FSceneView* View, FPrimitiveDrawInterface* PDI) override;
-	virtual bool ShouldOrbitCamera() const override;
-	bool CanSetWidgetMode(UE::Widget::EWidgetMode NewMode) const override;
-	void SetWidgetMode(UE::Widget::EWidgetMode NewMode) override;
-	UE::Widget::EWidgetMode GetWidgetMode() const override;	
-	void DrawCanvas(FViewport& InViewport, FSceneView& View, FCanvas& Canvas) override;
+	UE_API virtual void Draw(const FSceneView* View, FPrimitiveDrawInterface* PDI) override;
+	UE_API virtual bool ShouldOrbitCamera() const override;
+	UE_API bool CanSetWidgetMode(UE::Widget::EWidgetMode NewMode) const override;
+	UE_API void SetWidgetMode(UE::Widget::EWidgetMode NewMode) override;
+	UE_API UE::Widget::EWidgetMode GetWidgetMode() const override;	
+	UE_API void DrawCanvas(FViewport& InViewport, FSceneView& View, FCanvas& Canvas) override;
 
 
 	// Overriding base class visibility
 	using FEditorViewportClient::OverrideNearClipPlane;
 
 	// IInputBehaviorSource
-	virtual const UInputBehaviorSet* GetInputBehaviors() const override;
+	UE_API virtual const UInputBehaviorSet* GetInputBehaviors() const override;
 
 	// FGCObject
-	virtual void AddReferencedObjects(FReferenceCollector& Collector) override;
+	UE_API virtual void AddReferencedObjects(FReferenceCollector& Collector) override;
 
 protected:
-	void DrawGrid(const FSceneView* View, FPrimitiveDrawInterface* PDI);
-	void DrawGridRulers(FViewport& InViewport, FSceneView& View, UCanvas& Canvas);
-	void DrawUDIMLabels(FViewport& InViewport, FSceneView& View, UCanvas& Canvas);
+	UE_API void DrawGrid(const FSceneView* View, FPrimitiveDrawInterface* PDI);
+	UE_API void DrawGridRulers(FViewport& InViewport, FSceneView& View, UCanvas& Canvas);
+	UE_API void DrawUDIMLabels(FViewport& InViewport, FSceneView& View, UCanvas& Canvas);
 
 	// These get added in AddReferencedObjects for memory management
 	TObjectPtr<UInputBehaviorSet> BehaviorSet;
@@ -81,4 +83,6 @@ protected:
 	// anywhere in the header.
 	TUniquePtr<FEditor2DScrollBehaviorTarget> ScrollBehaviorTarget;
 	TUniquePtr<FEditor2DMouseWheelZoomBehaviorTarget> ZoomBehaviorTarget;
-};
\ No newline at end of file
+};
+
+#undef UE_API
diff --git a/Engine/Plugins/Editor/UVEditor/Source/UVEditor/Public/UVEditor3DViewportClient.h b/Engine/Plugins/Editor/UVEditor/Source/UVEditor/Public/UVEditor3DViewportClient.h
--- a/Engine/Plugins/Editor/UVEditor/Source/UVEditor/Public/UVEditor3DViewportClient.h
+++ b/Engine/Plugins/Editor/UVEditor/Source/UVEditor/Public/UVEditor3DViewportClient.h
@@ -6,6 +6,8 @@
 
 #include "EditorViewportClient.h"
 
+#define UE_API UVEDITOR_API
+
 class UUVToolViewportButtonsAPI;
 
 // Types of camera motion for the UV Editor 3D viewport
@@ -18,21 +20,21 @@ enum EUVEditor3DViewportClientCameraMode {
  * Viewport client for the 3d live preview in the UV editor. Currently same as editor viewport
  * client but doesn't allow editor gizmos/widgets, and alters orbit camera control.
  */
-class UVEDITOR_API FUVEditor3DViewportClient : public FEditorViewportClient
+class FUVEditor3DViewportClient : public FEditorViewportClient
 {
 public:
 
-	FUVEditor3DViewportClient(FEditorModeTools* InModeTools, FPreviewScene* InPreviewScene = nullptr,
+	UE_API FUVEditor3DViewportClient(FEditorModeTools* InModeTools, FPreviewScene* InPreviewScene = nullptr,
 		const TWeakPtr<SEditorViewport>& InEditorViewportWidget = nullptr, UUVToolViewportButtonsAPI* ViewportButtonsAPI = nullptr);
 
 	virtual ~FUVEditor3DViewportClient() {}
 
 	// FEditorViewportClient
-	virtual bool ShouldOrbitCamera() const override;
+	UE_API virtual bool ShouldOrbitCamera() const override;
 	bool CanSetWidgetMode(UE::Widget::EWidgetMode NewMode) const override {	return false; }
 	void SetWidgetMode(UE::Widget::EWidgetMode NewMode) override {}
 	UE::Widget::EWidgetMode GetWidgetMode() const override { return UE::Widget::EWidgetMode::WM_None; }
-	void FocusCameraOnSelection();
+	UE_API void FocusCameraOnSelection();
 public:
 
 	void SetCameraMode(EUVEditor3DViewportClientCameraMode CameraModeIn) { CameraMode = CameraModeIn; };
@@ -45,4 +47,6 @@ protected:
 	EUVEditor3DViewportClientCameraMode CameraMode = EUVEditor3DViewportClientCameraMode::Orbit;
 	UUVToolViewportButtonsAPI* ViewportButtonsAPI;
 
-};
\ No newline at end of file
+};
+
+#undef UE_API
diff --git a/Engine/Plugins/Editor/UVEditor/Source/UVEditor/Public/UVEditorBackgroundPreview.h b/Engine/Plugins/Editor/UVEditor/Source/UVEditor/Public/UVEditorBackgroundPreview.h
--- a/Engine/Plugins/Editor/UVEditor/Source/UVEditor/Public/UVEditorBackgroundPreview.h
+++ b/Engine/Plugins/Editor/UVEditor/Source/UVEditor/Public/UVEditorBackgroundPreview.h
@@ -9,6 +9,8 @@
 #include "GeometryBase.h"
 #include "UVEditorBackgroundPreview.generated.h"
 
+#define UE_API UVEDITOR_API
+
 
 /**
  * Enum to control the background visualiztion mode
@@ -24,8 +26,8 @@ enum class EUVEditorBackgroundSourceType
 /**
  * Visualization settings for the UUVEditorBackgroundPreview
  */
-UCLASS()
-class UVEDITOR_API UUVEditorBackgroundPreviewProperties : public UInteractiveToolPropertySet
+UCLASS(MinimalAPI)
+class UUVEditorBackgroundPreviewProperties : public UInteractiveToolPropertySet
 {
 	GENERATED_BODY()
 public:
@@ -56,8 +58,8 @@ public:
   Serves as a container for the background texture/material display in the UVEditor. This class is responsible for managing the quad
   drawn behind the grid, as well as maintaining the texture and material choices from the user to display.
  */
-UCLASS(Transient)
-class UVEDITOR_API UUVEditorBackgroundPreview : public UPreviewGeometry
+UCLASS(MinimalAPI, Transient)
+class UUVEditorBackgroundPreview : public UPreviewGeometry
 {
 	GENERATED_BODY()
 
@@ -66,7 +68,7 @@ public:
 	/**
 	 * Client must call this every frame for changes to .Settings to be reflected in rendered result.
 	 */
-	void OnTick(float DeltaTime);
+	UE_API void OnTick(float DeltaTime);
 	
 public:
 	/** Visualization settings */
@@ -85,11 +87,12 @@ public:
 	FOnBackgroundMaterialChange OnBackgroundMaterialChange;
 
 protected:
-	virtual void OnCreated() override;
+	UE_API virtual void OnCreated() override;
 
 	bool bSettingsModified = false;
 
-	void UpdateVisibility();
-	void UpdateBackground();
+	UE_API void UpdateVisibility();
+	UE_API void UpdateBackground();
 };
 
+#undef UE_API
