Fix localization text for solo column in ObjectMixer

diff --git a/Engine/Plugins/Editor/ObjectMixer/ObjectMixer/Source/ObjectMixer/Private/Views/List/Modes/OutlinerColumns/ObjectMixerOutlinerColumns.cpp b/Engine/Plugins/Editor/ObjectMixer/ObjectMixer/Source/ObjectMixer/Private/Views/List/Modes/OutlinerColumns/ObjectMixerOutlinerColumns.cpp
--- a/Engine/Plugins/Editor/ObjectMixer/ObjectMixer/Source/ObjectMixer/Private/Views/List/Modes/OutlinerColumns/ObjectMixerOutlinerColumns.cpp
+++ b/Engine/Plugins/Editor/ObjectMixer/ObjectMixer/Source/ObjectMixer/Private/Views/List/Modes/OutlinerColumns/ObjectMixerOutlinerColumns.cpp
@@ -43,7 +43,7 @@ SHeaderRow::FColumn::FArguments FObjectMixerOutlinerSoloColumn::ConstructHeaderR
 		.VAlignHeader(VAlign_Center)
 		.HAlignCell(HAlign_Center)
 		.VAlignCell(VAlign_Center)
-		.DefaultTooltip(FText::FromName(GetColumnID()))
+		.DefaultTooltip(LOCTEXT("SoloColumnTooltip", "Solo Objects"))
 		.HeaderComboVisibility(EHeaderComboVisibility::Never)
 		.HeaderContentPadding(FMargin(4.0f, 0.0f, 0.0f, 0.0f))
 		[
