Composure: Added integration with color grading drawer for composure color grading passes. Now, any color grading passes will show up in the drawer's mixer and can be selected and edited from the color grading panel. Object Mixer has been modified to support non-component sub-objects to support this feature

#Synced-CL 49006522

diff --git a/Engine/Plugins/Compositing/Composite/Composite.uplugin b/Engine/Plugins/Compositing/Composite/Composite.uplugin
--- a/Engine/Plugins/Compositing/Composite/Composite.uplugin
+++ b/Engine/Plugins/Compositing/Composite/Composite.uplugin
@@ -57,5 +57,15 @@
 			"TargetAllowList": [ "Editor" ],
 			"Enabled": true
 		},
+		{
+			"Name": "ObjectMixer",
+			"TargetAllowList": [ "Editor" ],
+			"Enabled": true
+		},
+		{
+			"Name": "ColorGrading",
+			"TargetAllowList": [ "Editor" ],
+			"Enabled": true
+		}
 	]
 }
\ No newline at end of file
diff --git a/Engine/Plugins/Compositing/Composite/Source/CompositeEditor/CompositeEditor.Build.cs b/Engine/Plugins/Compositing/Composite/Source/CompositeEditor/CompositeEditor.Build.cs
--- a/Engine/Plugins/Compositing/Composite/Source/CompositeEditor/CompositeEditor.Build.cs
+++ b/Engine/Plugins/Compositing/Composite/Source/CompositeEditor/CompositeEditor.Build.cs
@@ -20,6 +20,7 @@ public class CompositeEditor : ModuleRules
 			new string[]
 			{
 				"ApplicationCore",
+				"ColorGradingEditor",
 				"CompositeCore",
 				"ConcertSyncClient",
 				"CoreUObject",
@@ -29,6 +30,7 @@ public class CompositeEditor : ModuleRules
 				"LevelEditor",
 				"MediaAssets",
 				"MediaFrameworkUtilities",
+				"ObjectMixerEditor",
 				"Projects",
 				"PropertyEditor",
 				"UnrealEd",
diff --git a/Engine/Plugins/Compositing/Composite/Source/CompositeEditor/Private/ColorGradingDrawer/ColorGradingDataModelGenerator_Composite.cpp b/Engine/Plugins/Compositing/Composite/Source/CompositeEditor/Private/ColorGradingDrawer/ColorGradingDataModelGenerator_Composite.cpp
new file mode 100644
--- /dev/null
+++ b/Engine/Plugins/Compositing/Composite/Source/CompositeEditor/Private/ColorGradingDrawer/ColorGradingDataModelGenerator_Composite.cpp
@@ -0,0 +1,189 @@
+﻿// Copyright Epic Games, Inc. All Rights Reserved.
+
+#include "ColorGradingDataModelGenerator_Composite.h"
+
+#include "DetailCategoryBuilder.h"
+#include "DetailLayoutBuilder.h"
+#include "IDetailCustomization.h"
+#include "IDetailTreeNode.h"
+#include "IPropertyRowGenerator.h"
+#include "Passes/CompositePassColorGrading.h"
+#include "Widgets/Images/SImage.h"
+#include "Widgets/Text/STextBlock.h"
+
+#define LOCTEXT_NAMESPACE "FColorGradingDataModelGenerator_Composite"
+
+class FCompositePassColorGradingCustomization : public IDetailCustomization
+{
+public:
+	virtual void CustomizeDetails(IDetailLayoutBuilder& DetailBuilder) override
+	{
+		TArray<FName> Categories;
+		DetailBuilder.GetCategoryNames(Categories);
+
+		// Hide all default categories to avoid duplicate handles and unnecessary customization
+		for (const FName& Category : Categories)
+		{
+			DetailBuilder.HideCategory(Category);
+		}
+		
+		IDetailCategoryBuilder& ColorGradingHeaderCategory = DetailBuilder.EditCategory(TEXT("ColorGradingHeader"));
+		ColorGradingHeaderCategory.AddProperty(DetailBuilder.GetProperty(TEXT("bIsEnabled"), UCompositePassBase::StaticClass()));
+		
+		IDetailCategoryBuilder& ColorGradingElementsCategory = DetailBuilder.EditCategory(TEXT("ColorGradingElements"));
+
+		ColorGradingElementsCategory.AddProperty(DetailBuilder.GetProperty(GET_MEMBER_NAME_CHECKED(UCompositePassColorGrading, ColorGradingSettings.Global)));
+		ColorGradingElementsCategory.AddProperty(DetailBuilder.GetProperty(GET_MEMBER_NAME_CHECKED(UCompositePassColorGrading, ColorGradingSettings.Shadows)));
+		ColorGradingElementsCategory.AddProperty(DetailBuilder.GetProperty(GET_MEMBER_NAME_CHECKED(UCompositePassColorGrading, ColorGradingSettings.Midtones)));
+		ColorGradingElementsCategory.AddProperty(DetailBuilder.GetProperty(GET_MEMBER_NAME_CHECKED(UCompositePassColorGrading, ColorGradingSettings.Highlights)));
+
+		IDetailCategoryBuilder& TemperatureCategory = DetailBuilder.EditCategory(TEXT("Details_Temperature"), LOCTEXT("TemperatureCategory", "Temperature"));
+		TemperatureCategory.AddProperty(DetailBuilder.GetProperty(GET_MEMBER_NAME_CHECKED(UCompositePassColorGrading, TemperatureSettings.TemperatureType)));
+		TemperatureCategory.AddProperty(DetailBuilder.GetProperty(GET_MEMBER_NAME_CHECKED(UCompositePassColorGrading, TemperatureSettings.WhiteTemp)));
+		TemperatureCategory.AddProperty(DetailBuilder.GetProperty(GET_MEMBER_NAME_CHECKED(UCompositePassColorGrading, TemperatureSettings.WhiteTint)));
+
+		IDetailCategoryBuilder& ColorGradingCategory = DetailBuilder.EditCategory(TEXT("Details_ColorGrading"), LOCTEXT("ColorGradingCategory", "Color Grading"));
+		ColorGradingCategory.AddProperty(DetailBuilder.GetProperty(GET_MEMBER_NAME_CHECKED(UCompositePassColorGrading, ColorGradingSettings.ShadowsMax)));
+		ColorGradingCategory.AddProperty(DetailBuilder.GetProperty(GET_MEMBER_NAME_CHECKED(UCompositePassColorGrading, ColorGradingSettings.HighlightsMin)));
+		ColorGradingCategory.AddProperty(DetailBuilder.GetProperty(GET_MEMBER_NAME_CHECKED(UCompositePassColorGrading, ColorGradingSettings.HighlightsMax)));
+	}
+};
+
+void FColorGradingDataModelGenerator_Composite::Initialize(const TSharedRef<FColorGradingEditorDataModel>& ColorGradingDataModel, const TSharedRef<IPropertyRowGenerator>& PropertyRowGenerator)
+{
+	PropertyRowGenerator->RegisterInstancedCustomPropertyLayout(UCompositePassColorGrading::StaticClass(), FOnGetDetailCustomizationInstance::CreateLambda([]
+	{
+		return MakeShared<FCompositePassColorGradingCustomization>();
+	}));
+}
+
+void FColorGradingDataModelGenerator_Composite::Destroy(const TSharedRef<FColorGradingEditorDataModel>& ColorGradingDataModel, const TSharedRef<IPropertyRowGenerator>& PropertyRowGenerator)
+{
+	PropertyRowGenerator->UnregisterInstancedCustomPropertyLayout(UCompositePassColorGrading::StaticClass());
+}
+
+void FColorGradingDataModelGenerator_Composite::GenerateDataModel(IPropertyRowGenerator& PropertyRowGenerator, FColorGradingEditorDataModel& OutColorGradingDataModel)
+{
+	TArray<TWeakObjectPtr<UCompositePassColorGrading>> SelectedPasses;
+	const TArray<TWeakObjectPtr<UObject>>& SelectedObjects = PropertyRowGenerator.GetSelectedObjects();
+	for (const TWeakObjectPtr<UObject>& SelectedObject : SelectedObjects)
+	{
+		if (SelectedObject.IsValid() && SelectedObject->IsA<UCompositePassColorGrading>())
+		{
+			TWeakObjectPtr<UCompositePassColorGrading> SelectedRootActor = CastChecked<UCompositePassColorGrading>(SelectedObject.Get());
+			SelectedPasses.Add(SelectedRootActor);
+		}
+	}
+
+	if (SelectedPasses.IsEmpty())
+	{
+		return;
+	}
+
+	const TArray<TSharedRef<IDetailTreeNode>>& RootNodes = PropertyRowGenerator.GetRootTreeNodes();
+
+	const TSharedRef<IDetailTreeNode>* ColorGradingElementsPtr = RootNodes.FindByPredicate([](const TSharedRef<IDetailTreeNode>& Node)
+	{
+		return Node->GetNodeName() == TEXT("ColorGradingElements");
+	});
+
+	if (ColorGradingElementsPtr)
+	{
+		const TSharedRef<IDetailTreeNode> ColorGradingElements = *ColorGradingElementsPtr;
+		FColorGradingEditorDataModel::FColorGradingGroup ColorGradingGroup;
+
+		ColorGradingGroup.DetailsViewCategories.Append(
+		{
+			TEXT("Details_Temperature"),
+			TEXT("Details_ColorGrading")
+		});
+		
+		TArray<TSharedRef<IDetailTreeNode>> ColorGradingPropertyNodes;
+		ColorGradingElements->GetChildren(ColorGradingPropertyNodes);
+
+		for (const TSharedRef<IDetailTreeNode>& PropertyNode : ColorGradingPropertyNodes)
+		{
+			TSharedPtr<IPropertyHandle> PropertyHandle = PropertyNode->CreatePropertyHandle();
+
+			FColorGradingEditorDataModel::FColorGradingElement ColorGradingElement = CreateColorGradingElement(PropertyNode, FText::FromName(PropertyNode->GetNodeName()));
+			ColorGradingGroup.ColorGradingElements.Add(ColorGradingElement);
+		}
+
+		const TSharedRef<IDetailTreeNode>* ColorGradingHeaderPtr = RootNodes.FindByPredicate([](const TSharedRef<IDetailTreeNode>& Node)
+		{
+			return Node->GetNodeName() == TEXT("ColorGradingHeader");
+		});
+
+		if (ColorGradingHeaderPtr)
+		{
+			const TSharedRef<IDetailTreeNode> ColorGradingHeader = *ColorGradingHeaderPtr;
+			
+			TArray<TSharedRef<IDetailTreeNode>> ColorGradingHeaderProperties;
+			ColorGradingHeader->GetChildren(ColorGradingHeaderProperties);
+
+			if (ColorGradingHeaderProperties.Num() > 0)
+			{
+				ColorGradingGroup.EditConditionPropertyHandle = ColorGradingHeaderProperties[0]->CreatePropertyHandle();
+				ColorGradingGroup.GroupHeaderWidget =
+					SNew(SHorizontalBox)
+					+SHorizontalBox::Slot()
+					.AutoWidth()
+					.Padding(0, 0, 4, 0)
+					.VAlign(VAlign_Center)
+					[
+						SNew(STextBlock)
+						.Text(LOCTEXT("EnabledColorGradingPassLabel", "Enable Color Grading Pass"))
+						.Font(FAppStyle::Get().GetFontStyle("NormalFontBold"))
+					];
+			}
+		}
+		
+		OutColorGradingDataModel.ColorGradingGroups.Add(ColorGradingGroup);
+	}
+}
+
+FColorGradingEditorDataModel::FColorGradingElement FColorGradingDataModelGenerator_Composite::CreateColorGradingElement(const TSharedRef<IDetailTreeNode>& GroupNode, FText ElementLabel)
+{
+	FColorGradingEditorDataModel::FColorGradingElement ColorGradingElement;
+	ColorGradingElement.DisplayName = ElementLabel;
+
+	TArray<TSharedRef<IDetailTreeNode>> ChildNodes;
+	GroupNode->GetChildren(ChildNodes);
+
+	for (const TSharedRef<IDetailTreeNode>& ChildNode : ChildNodes)
+	{
+		TSharedPtr<IPropertyHandle> PropertyHandle = ChildNode->CreatePropertyHandle();
+		if (PropertyHandle.IsValid() && PropertyHandle->IsValidHandle())
+		{
+			const FString ColorGradingModeString = PropertyHandle->GetProperty()->GetMetaData(TEXT("ColorGradingMode")).ToLower();
+
+			if (!ColorGradingModeString.IsEmpty())
+			{
+				if (ColorGradingModeString.Compare(TEXT("saturation")) == 0)
+				{
+					ColorGradingElement.SaturationPropertyHandle = PropertyHandle;
+				}
+				else if (ColorGradingModeString.Compare(TEXT("contrast")) == 0)
+				{
+					ColorGradingElement.ContrastPropertyHandle = PropertyHandle;
+				}
+				else if (ColorGradingModeString.Compare(TEXT("gamma")) == 0)
+				{
+					ColorGradingElement.GammaPropertyHandle = PropertyHandle;
+				}
+				else if (ColorGradingModeString.Compare(TEXT("gain")) == 0)
+				{
+					ColorGradingElement.GainPropertyHandle = PropertyHandle;
+				}
+				else if (ColorGradingModeString.Compare(TEXT("offset")) == 0)
+				{
+					ColorGradingElement.OffsetPropertyHandle = PropertyHandle;
+				}
+			}
+		}
+	}
+
+	return ColorGradingElement;
+}
+
+#undef LOCTEXT_NAMESPACE
\ No newline at end of file
diff --git a/Engine/Plugins/Compositing/Composite/Source/CompositeEditor/Private/ColorGradingDrawer/ColorGradingDataModelGenerator_Composite.h b/Engine/Plugins/Compositing/Composite/Source/CompositeEditor/Private/ColorGradingDrawer/ColorGradingDataModelGenerator_Composite.h
new file mode 100644
--- /dev/null
+++ b/Engine/Plugins/Compositing/Composite/Source/CompositeEditor/Private/ColorGradingDrawer/ColorGradingDataModelGenerator_Composite.h
@@ -0,0 +1,22 @@
+﻿// Copyright Epic Games, Inc. All Rights Reserved.
+
+#pragma once
+
+#include "ColorGradingEditorDataModel.h"
+
+/** Color drawer data model generator for UCompositePassColorGrading objects*/
+class FColorGradingDataModelGenerator_Composite : public IColorGradingEditorDataModelGenerator
+{
+public:
+	static TSharedRef<IColorGradingEditorDataModelGenerator> MakeInstance() { return MakeShareable(new FColorGradingDataModelGenerator_Composite());}
+	
+	//~ IColorGradingDataModelGenerator interface
+	virtual void Initialize(const TSharedRef<FColorGradingEditorDataModel>& ColorGradingDataModel, const TSharedRef<IPropertyRowGenerator>& PropertyRowGenerator) override;
+	virtual void Destroy(const TSharedRef<FColorGradingEditorDataModel>& ColorGradingDataModel, const TSharedRef<IPropertyRowGenerator>& PropertyRowGenerator) override;
+	virtual void GenerateDataModel(IPropertyRowGenerator& PropertyRowGenerator, FColorGradingEditorDataModel& OutColorGradingDataModel) override;
+	//~ End IColorGradingDataModelGenerator interface
+
+private:
+	/** Creates a new color grading element structure for the specified detail tree node, which is expected to have child color properties with the ColorGradingMode metadata set */
+	FColorGradingEditorDataModel::FColorGradingElement CreateColorGradingElement(const TSharedRef<IDetailTreeNode>& GroupNode, FText ElementLabel);
+};
diff --git a/Engine/Plugins/Compositing/Composite/Source/CompositeEditor/Private/ColorGradingDrawer/ColorGradingHierarchyConfig_Composite.cpp b/Engine/Plugins/Compositing/Composite/Source/CompositeEditor/Private/ColorGradingDrawer/ColorGradingHierarchyConfig_Composite.cpp
new file mode 100644
--- /dev/null
+++ b/Engine/Plugins/Compositing/Composite/Source/CompositeEditor/Private/ColorGradingDrawer/ColorGradingHierarchyConfig_Composite.cpp
@@ -0,0 +1,65 @@
+﻿// Copyright Epic Games, Inc. All Rights Reserved.
+
+#include "ColorGradingHierarchyConfig_Composite.h"
+
+#include "CompositeActor.h"
+#include "CompositeEditorStyle.h"
+#include "Layers/CompositeLayerPlate.h"
+#include "ObjectFilter/ObjectMixerEditorActorSubObject.h"
+#include "Passes/CompositePassColorGrading.h"
+
+TArray<FObjectMixerEditorActorSubObject> FColorGradingHierarchyConfig_Composite::GetActorSubObjects(UObject* ParentObject) const
+{
+	TArray<FObjectMixerEditorActorSubObject> SubObjects;
+	
+	if (const ACompositeActor* CompositeActor = Cast<ACompositeActor>(ParentObject))
+	{
+		for (const TObjectPtr<UCompositeLayerBase>& CompositeLayer : CompositeActor->GetCompositeLayers())
+		{
+			const UCompositeLayerPlate* PlateLayer = Cast<UCompositeLayerPlate>(CompositeLayer.Get());
+			if (IsValid(PlateLayer))
+			{
+				for (const TObjectPtr<UCompositePassBase>& Pass : PlateLayer->MediaPasses)
+				{
+					if (IsValid(Pass) && Pass->IsA<UCompositePassColorGrading>())
+					{
+						FObjectMixerEditorActorSubObject SubObject;
+						SubObject.Object = Pass.Get();
+						SubObject.DisplayName = FText::FromString(Pass->GetDisplayName());
+						SubObject.bIncludeObjectStack = true;
+						SubObject.ContainerProperty.PropertyName = GET_MEMBER_NAME_CHECKED(UCompositeLayerPlate, MediaPasses);
+						SubObject.ContainerProperty.PropertyIconOverride = FCompositeEditorStyle::Get().GetBrush("CompositeEditor.Passes.Media");
+						
+						SubObjects.Add(SubObject);
+					}
+				}
+
+				for (const TObjectPtr<UCompositePassBase>& Pass : PlateLayer->LayerPasses)
+				{
+					if (IsValid(Pass) && Pass->IsA<UCompositePassColorGrading>())
+					{
+						FObjectMixerEditorActorSubObject SubObject;
+						SubObject.Object = Pass.Get();
+						SubObject.DisplayName = FText::FromString(Pass->GetDisplayName());
+						SubObject.bIncludeObjectStack = true;
+						SubObject.ContainerProperty.PropertyName = GET_MEMBER_NAME_CHECKED(UCompositeLayerPlate, LayerPasses);
+						SubObject.ContainerProperty.PropertyIconOverride = FCompositeEditorStyle::Get().GetBrush("CompositeEditor.Passes.Layer");
+						
+						SubObjects.Add(SubObject);
+					}
+				}
+			}
+		}
+	}
+
+	return MoveTemp(SubObjects);
+}
+
+TSet<FName> FColorGradingHierarchyConfig_Composite::GetPropertiesThatRequireListRefresh() const
+{
+	return
+	{
+		GET_MEMBER_NAME_CHECKED(UCompositeLayerPlate, MediaPasses),
+		GET_MEMBER_NAME_CHECKED(UCompositeLayerPlate, LayerPasses)
+	};
+}
diff --git a/Engine/Plugins/Compositing/Composite/Source/CompositeEditor/Private/ColorGradingDrawer/ColorGradingHierarchyConfig_Composite.h b/Engine/Plugins/Compositing/Composite/Source/CompositeEditor/Private/ColorGradingDrawer/ColorGradingHierarchyConfig_Composite.h
new file mode 100644
--- /dev/null
+++ b/Engine/Plugins/Compositing/Composite/Source/CompositeEditor/Private/ColorGradingDrawer/ColorGradingHierarchyConfig_Composite.h
@@ -0,0 +1,17 @@
+﻿// Copyright Epic Games, Inc. All Rights Reserved.
+
+#pragma once
+
+#include "ColorGradingMixerObjectFilterRegistry.h"
+
+/** Color grading mixer hierarchy config for displaying UCompositePassColorGrading objects in the color grading drawer */
+struct FColorGradingHierarchyConfig_Composite : IColorGradingMixerObjectHierarchyConfig
+{
+public:
+	static TSharedRef<IColorGradingMixerObjectHierarchyConfig> MakeInstance() { return MakeShared<FColorGradingHierarchyConfig_Composite>(); }
+
+	//~ IColorGradingMixerObjectHierarchyConfig interface
+	virtual TArray<FObjectMixerEditorActorSubObject> GetActorSubObjects(UObject* ParentObject) const override;
+	virtual TSet<FName> GetPropertiesThatRequireListRefresh() const override;
+	//~ End IColorGradingMixerObjectHierarchyConfig interface
+};
diff --git a/Engine/Plugins/Compositing/Composite/Source/CompositeEditor/Private/ColorGradingDrawer/ColorGradingSelectionHandler_Composite.cpp b/Engine/Plugins/Compositing/Composite/Source/CompositeEditor/Private/ColorGradingDrawer/ColorGradingSelectionHandler_Composite.cpp
new file mode 100644
--- /dev/null
+++ b/Engine/Plugins/Compositing/Composite/Source/CompositeEditor/Private/ColorGradingDrawer/ColorGradingSelectionHandler_Composite.cpp
@@ -0,0 +1,51 @@
+﻿// Copyright Epic Games, Inc. All Rights Reserved.
+
+#include "ColorGradingSelectionHandler_Composite.h"
+
+#include "CompositeActor.h"
+#include "Passes/CompositePassColorGrading.h"
+#include "UI/SCompositeEditorPanel.h"
+
+TSharedRef<IColorGradingMixerSelectionHandler> FColorGradingSelectionHandler_Composite::MakeInstance()
+{
+	TSharedRef<FColorGradingSelectionHandler_Composite> SelectionHandler = MakeShared<FColorGradingSelectionHandler_Composite>();
+
+	SCompositeEditorPanel::GetOnSelectionChanged().AddSP(SelectionHandler, &FColorGradingSelectionHandler_Composite::OnComposurePanelSelectionChanged);
+
+	return SelectionHandler;
+}
+
+bool FColorGradingSelectionHandler_Composite::CanSelectSubObject(const UObject* InObject)
+{
+	return InObject->IsA<ACompositeActor>() || InObject->IsA<UCompositeLayerBase>() || InObject->IsA<UCompositePassBase>();
+}
+
+void FColorGradingSelectionHandler_Composite::SelectSubObjects(const TArray<UObject*>& InSelectedObjects, bool bShouldSelect, bool bSelectEvenIfHidden) const
+{
+	if (bShouldSelect)
+	{
+		SCompositeEditorPanel::UpdateActivePanelSelection(InSelectedObjects);
+	}
+	else
+	{
+		SCompositeEditorPanel::UpdateActivePanelSelection({ });
+	}
+}
+
+TArray<UObject*> FColorGradingSelectionHandler_Composite::GetSelectedSubObjectsInEditor() const
+{
+	TArray<UObject*> SelectedObjects = SCompositeEditorPanel::GetActivePanelSelection();
+
+	// Remove any non-color grading objects from the selection
+	SelectedObjects.RemoveAll([](const UObject* Object)
+	{
+		return !Object->IsA<UCompositePassColorGrading>();
+	});
+
+	return SelectedObjects;
+}
+
+void FColorGradingSelectionHandler_Composite::OnComposurePanelSelectionChanged()
+{
+	OnSelectionChanged.ExecuteIfBound();
+}
diff --git a/Engine/Plugins/Compositing/Composite/Source/CompositeEditor/Private/ColorGradingDrawer/ColorGradingSelectionHandler_Composite.h b/Engine/Plugins/Compositing/Composite/Source/CompositeEditor/Private/ColorGradingDrawer/ColorGradingSelectionHandler_Composite.h
new file mode 100644
--- /dev/null
+++ b/Engine/Plugins/Compositing/Composite/Source/CompositeEditor/Private/ColorGradingDrawer/ColorGradingSelectionHandler_Composite.h
@@ -0,0 +1,21 @@
+﻿// Copyright Epic Games, Inc. All Rights Reserved.
+
+#pragma once
+
+#include "ColorGradingMixerObjectFilterRegistry.h"
+
+/** Color grading selection handler that synchronizes selection between the color grading drawer and the Composure panel */
+struct FColorGradingSelectionHandler_Composite : IColorGradingMixerSelectionHandler
+{
+public:
+	static TSharedRef<IColorGradingMixerSelectionHandler> MakeInstance();
+
+	//~ IColorGradingMixerSelectionHandler interface
+	virtual  bool CanSelectSubObject(const UObject* InObject) override;
+	virtual void SelectSubObjects(const TArray<UObject*>& InSelectedObjects, bool bShouldSelect, bool bSelectEvenIfHidden) const override;
+	virtual TArray<UObject*> GetSelectedSubObjectsInEditor() const override;
+	//~ End IColorGradingMixerSelectionHandler interface
+	
+private:
+	void OnComposurePanelSelectionChanged();
+};
diff --git a/Engine/Plugins/Compositing/Composite/Source/CompositeEditor/Private/CompositeEditorModule.cpp b/Engine/Plugins/Compositing/Composite/Source/CompositeEditor/Private/CompositeEditorModule.cpp
--- a/Engine/Plugins/Compositing/Composite/Source/CompositeEditor/Private/CompositeEditorModule.cpp
+++ b/Engine/Plugins/Compositing/Composite/Source/CompositeEditor/Private/CompositeEditorModule.cpp
@@ -2,6 +2,8 @@
 
 #include "CompositeEditorModule.h"
 
+#include "ColorGradingEditorDataModel.h"
+#include "ColorGradingMixerObjectFilterRegistry.h"
 #include "Editor.h"
 #include "EngineUtils.h"
 #include "HoldoutCompositeComponent.h"
@@ -17,6 +19,9 @@
 #include "CompositeEditorCommands.h"
 #include "CompositeEditorStyle.h"
 #include "PropertyEditorModule.h"
+#include "ColorGradingDrawer/ColorGradingDataModelGenerator_Composite.h"
+#include "ColorGradingDrawer/ColorGradingHierarchyConfig_Composite.h"
+#include "ColorGradingDrawer/ColorGradingSelectionHandler_Composite.h"
 #include "Customizations/CompositeActorCustomization.h"
 #include "Customizations/CompositeLayerPlateCustomization.h"
 #include "Customizations/CompositeLayerSceneCaptureCustomization.h"
@@ -27,7 +32,7 @@
 #include "Layers/CompositeLayerSceneCapture.h"
 #include "Layers/CompositeLayerShadowReflection.h"
 #include "Layers/CompositeLayerSingleLightShadow.h"
-#include "Misc/MessageDialog.h"
+#include "Passes/CompositePassColorGrading.h"
 #include "UI/SCompositeEditorPanel.h"
 
 #define LOCTEXT_NAMESPACE "CompositeEditorModule"
@@ -45,6 +50,18 @@ void FCompositeEditorModule::StartupModule()
 	
 	RegisterTabSpawners();
 	RegisterCustomizations();
+
+	FColorGradingEditorDataModel::RegisterColorGradingDataModelGenerator<UCompositePassColorGrading>(
+		FGetDetailsDataModelGenerator::CreateStatic(&FColorGradingDataModelGenerator_Composite::MakeInstance));
+	
+	FColorGradingMixerObjectFilterRegistry::RegisterActorClassToPlace(ACompositeActor::StaticClass());
+
+	FColorGradingMixerObjectFilterRegistry::RegisterObjectClassToFilter(
+		ACompositeActor::StaticClass(),
+		FGetObjectHierarchyConfig::CreateStatic(&FColorGradingHierarchyConfig_Composite::MakeInstance)
+	);
+
+	FColorGradingMixerObjectFilterRegistry::RegisterSelectionHandler(FGetSelectionHandler::CreateStatic(&FColorGradingSelectionHandler_Composite::MakeInstance));
 }
 
 void FCompositeEditorModule::ShutdownModule()
diff --git a/Engine/Plugins/Compositing/Composite/Source/CompositeEditor/Private/CompositeEditorStyle.cpp b/Engine/Plugins/Compositing/Composite/Source/CompositeEditor/Private/CompositeEditorStyle.cpp
--- a/Engine/Plugins/Compositing/Composite/Source/CompositeEditor/Private/CompositeEditorStyle.cpp
+++ b/Engine/Plugins/Compositing/Composite/Source/CompositeEditor/Private/CompositeEditorStyle.cpp
@@ -41,7 +41,7 @@ FCompositeEditorStyle::FCompositeEditorStyle()
 	Set("ClassIcon.CompositeLayerSingleLightShadow", new IMAGE_BRUSH_SVG("Editor/Icons/Shadow", Icon16x16));
 	Set("ClassIcon.CompositeLayerSceneCapture", new CORE_IMAGE_BRUSH_SVG("Starship/AssetIcons/SceneCapture2D_16", Icon16x16));
 	Set("ClassIcon.CompositeLayerPlate", new CORE_IMAGE_BRUSH_SVG("Starship/AssetIcons/MediaPlayer_16", Icon16x16));
-	Set("ClassIcon.CompositePassColorGrade", new CORE_IMAGE_BRUSH_SVG("Starship/TimelineEditor/TrackTypeColor", Icon16x16));
+	Set("ClassIcon.CompositePassColorGrading", new CORE_IMAGE_BRUSH_SVG("Starship/TimelineEditor/TrackTypeColor", Icon16x16));
 	Set("ClassIcon.CompositePassColorKeyer", new IMAGE_BRUSH_SVG("Editor/Icons/Chromakey", Icon16x16));
 	Set("ClassIcon.CompositePassFXAA", new CORE_IMAGE_BRUSH_SVG("Starship/Common/AntiAliasing", Icon16x16));
 	Set("ClassIcon.CompositePassSMAA", new CORE_IMAGE_BRUSH_SVG("Starship/Common/AntiAliasing", Icon16x16));
diff --git a/Engine/Plugins/Compositing/Composite/Source/CompositeEditor/Private/Customizations/CompositeLayerPlateCustomization.cpp b/Engine/Plugins/Compositing/Composite/Source/CompositeEditor/Private/Customizations/CompositeLayerPlateCustomization.cpp
--- a/Engine/Plugins/Compositing/Composite/Source/CompositeEditor/Private/Customizations/CompositeLayerPlateCustomization.cpp
+++ b/Engine/Plugins/Compositing/Composite/Source/CompositeEditor/Private/Customizations/CompositeLayerPlateCustomization.cpp
@@ -96,7 +96,7 @@ void FCompositeLayerPlateCustomization::CustomizeDetails(IDetailLayoutBuilder& D
 		ActorListGroup.AddWidgetRow()
 		.WholeRowContent()
 		[
-			SNew(SCompositeActorPickerTable, ActorListRef)
+			SAssignNew(CompositeMeshPicker, SCompositeActorPickerTable, ActorListRef)
 			.OnExtendAddMenu(this, &FCompositeLayerPlateCustomization::OnExtendCompositeMeshAddMenu)
 			.OnLayoutSizeChanged(this, &FCompositeLayerPlateCustomization::OnLayoutSizeChanged)
 			.ShowApplyMaterialSection(true)
@@ -107,7 +107,7 @@ void FCompositeLayerPlateCustomization::CustomizeDetails(IDetailLayoutBuilder& D
 		PassesGroup.AddWidgetRow()
 		.WholeRowContent()
 		[
-			SNew(SCompositePlatePassPanel, Plate)
+			SAssignNew(PassPanel, SCompositePlatePassPanel, Plate)
 			.OnLayoutSizeChanged(this, &FCompositeLayerPlateCustomization::OnLayoutSizeChanged)
 		];
 	}  
@@ -130,6 +130,29 @@ void FCompositeLayerPlateCustomization::CustomizeDetails(IDetailLayoutBuilder& D
 	}
 }
 
+bool FCompositeLayerPlateCustomization::IsCustomizingObject(UObject* InObject) const
+{
+	if (TSharedPtr<IDetailLayoutBuilder> Builder = CachedDetailBuilder.Pin())
+	{
+		TArray<TWeakObjectPtr<UObject>> ObjectsBeingCustomized;
+		Builder->GetObjectsBeingCustomized(ObjectsBeingCustomized);
+
+		return ObjectsBeingCustomized.Contains(InObject);
+	}
+
+	return false;
+}
+
+TSharedPtr<SCompositePlatePassPanel> FCompositeLayerPlateCustomization::GetPassPanelWidget() const
+{
+	return PassPanel.Pin();
+}
+
+TSharedPtr<SCompositeActorPickerTable> FCompositeLayerPlateCustomization::GetCompositeMeshPickerWidget() const
+{
+	return CompositeMeshPicker.Pin();
+}
+
 void FCompositeLayerPlateCustomization::CustomizeTexturePropertyRow(const TSharedPtr<IPropertyHandle>& InPropertyHandle, IDetailPropertyRow& InPropertyRow)
 {
 	TexturePropertyHandle = InPropertyHandle;
diff --git a/Engine/Plugins/Compositing/Composite/Source/CompositeEditor/Private/Customizations/CompositeLayerPlateCustomization.h b/Engine/Plugins/Compositing/Composite/Source/CompositeEditor/Private/Customizations/CompositeLayerPlateCustomization.h
--- a/Engine/Plugins/Compositing/Composite/Source/CompositeEditor/Private/Customizations/CompositeLayerPlateCustomization.h
+++ b/Engine/Plugins/Compositing/Composite/Source/CompositeEditor/Private/Customizations/CompositeLayerPlateCustomization.h
@@ -7,6 +7,7 @@
 class SWidget;
 class IDetailPropertyRow;
 class IPropertyHandle;
+class SCompositeActorPickerTable;
 class SCompositePlatePassPanel;
 class FMenuBuilder;
 
@@ -24,6 +25,10 @@ public:
 	virtual void CustomizeDetails(IDetailLayoutBuilder& DetailLayout) override;
 	// End of IDetailCustomization interface
 
+	bool IsCustomizingObject(UObject* InObject) const;
+	TSharedPtr<SCompositePlatePassPanel> GetPassPanelWidget() const;
+	TSharedPtr<SCompositeActorPickerTable> GetCompositeMeshPickerWidget() const;
+	
 private:
 	/** Customizes the Texture property in UCompositeLayerPlate to add in the media profile source selector */
 	void CustomizeTexturePropertyRow(const TSharedPtr<IPropertyHandle>& InPropertyHandle, IDetailPropertyRow& InPropertyRow);
@@ -54,5 +59,7 @@ private:
 	
 private:
 	TWeakPtr<IDetailLayoutBuilder> CachedDetailBuilder;
+	TWeakPtr<SCompositePlatePassPanel> PassPanel;
+	TWeakPtr<SCompositeActorPickerTable> CompositeMeshPicker;
 	TSharedPtr<IPropertyHandle> TexturePropertyHandle;
 };
diff --git a/Engine/Plugins/Compositing/Composite/Source/CompositeEditor/Private/UI/SCompositeEditorPanel.cpp b/Engine/Plugins/Compositing/Composite/Source/CompositeEditor/Private/UI/SCompositeEditorPanel.cpp
--- a/Engine/Plugins/Compositing/Composite/Source/CompositeEditor/Private/UI/SCompositeEditorPanel.cpp
+++ b/Engine/Plugins/Compositing/Composite/Source/CompositeEditor/Private/UI/SCompositeEditorPanel.cpp
@@ -11,9 +11,12 @@
 #include "IDetailCustomization.h"
 #include "LevelEditor.h"
 #include "SCompositePanelLayerTree.h"
+#include "SCompositePlatePassPanel.h"
 #include "ScopedTransaction.h"
 #include "WorkspaceMenuStructure.h"
 #include "WorkspaceMenuStructureModule.h"
+#include "Customizations/CompositeLayerPlateCustomization.h"
+#include "Layers/CompositeLayerPlate.h"
 #include "Widgets/Docking/SDockTab.h"
 #include "Widgets/Images/SImage.h"
 #include "Widgets/Input/SComboBox.h"
@@ -409,6 +412,38 @@ void SCompositeEditorPanel::UnregisterTabSpawner()
 	}
 }
 
+void SCompositeEditorPanel::UpdateActivePanelSelection(const TArray<UObject*>& InSelection)
+{
+	FLevelEditorModule& LevelEditorModule = FModuleManager::LoadModuleChecked<FLevelEditorModule>("LevelEditor");
+	if (TSharedPtr<FTabManager> LevelEditorTabManager = LevelEditorModule.GetLevelEditorTabManager())
+	{
+		if (TSharedPtr<SDockTab> Tab = LevelEditorTabManager->TryInvokeTab(CompositingEditorPanel::CompositeEditorTabName))
+		{
+			if (TSharedPtr<SCompositeEditorPanel> CompositePanel = StaticCastSharedPtr<SCompositeEditorPanel>(Tab->GetContent().ToSharedPtr()))
+			{
+				CompositePanel->SelectObjects(InSelection);
+			}
+		}
+	}
+}
+
+TArray<UObject*> SCompositeEditorPanel::GetActivePanelSelection()
+{
+	FLevelEditorModule& LevelEditorModule = FModuleManager::LoadModuleChecked<FLevelEditorModule>("LevelEditor");
+	if (TSharedPtr<FTabManager> LevelEditorTabManager = LevelEditorModule.GetLevelEditorTabManager())
+	{
+		if (TSharedPtr<SDockTab> Tab = LevelEditorTabManager->TryInvokeTab(CompositingEditorPanel::CompositeEditorTabName))
+		{
+			if (TSharedPtr<SCompositeEditorPanel> CompositePanel = StaticCastSharedPtr<SCompositeEditorPanel>(Tab->GetContent().ToSharedPtr()))
+			{
+				return CompositePanel->GetSelectedObjects();
+			}
+		}
+	}
+
+	return { };
+}
+
 void SCompositeEditorPanel::Construct(const FArguments& InArgs)
 {
 	FDetailsViewArgs DetailsViewArgs;
@@ -421,6 +456,10 @@ void SCompositeEditorPanel::Construct(const FArguments& InArgs)
 	{
 		return MakeShared<FCompositeActorPanelDetailCustomization>();
 	}));
+
+	// Create an instanced version of the plate customization so that a pointer to it can be stored and accessed later
+	DetailsView->RegisterInstancedCustomPropertyLayout(UCompositeLayerPlate::StaticClass(),
+		FOnGetDetailCustomizationInstance::CreateSP(this, &SCompositeEditorPanel::CreateCustomization<UCompositeLayerPlate, FCompositeLayerPlateCustomization>));
 	
 	ChildSlot
 	[
@@ -451,12 +490,127 @@ void SCompositeEditorPanel::SelectCompositeActors(const TArray<TWeakObjectPtr<AC
 	}
 }
 
+void SCompositeEditorPanel::SelectObjects(const TArray<UObject*>& InSelection)
+{
+	TArray<UObject*> CompositeActorsAndLayers;
+	TMap<UCompositeLayerBase*, TArray<UCompositePassBase*>> CompositePasses;
+
+	for (UObject* Object : InSelection)
+	{
+		if (!Object)
+		{
+			continue;
+		}
+		
+		if (Object->IsA<ACompositeActor>() || Object->IsA<UCompositeLayerBase>())
+		{
+			CompositeActorsAndLayers.Add(Object);
+		}
+
+		if (UCompositePassBase* Pass = Cast<UCompositePassBase>(Object))
+		{
+			if (UCompositeLayerBase* ParentLayer = Pass->GetTypedOuter<UCompositeLayerBase>())
+			{
+				CompositeActorsAndLayers.Add(ParentLayer);
+
+				if (!CompositePasses.Contains(ParentLayer))
+				{
+					CompositePasses.Add(ParentLayer);
+				}
+				
+				CompositePasses[ParentLayer].Add(Pass);
+			}
+		}
+	}
+	
+	if (LayerTree.IsValid())
+	{
+		LayerTree->SelectObjects(CompositeActorsAndLayers);
+	}
+
+	for (const TPair<UCompositeLayerBase*, TArray<UCompositePassBase*>>& PassesForLayer : CompositePasses)
+	{
+		UCompositeLayerBase* Layer = PassesForLayer.Key;
+		TArray<UCompositePassBase*> Passes = PassesForLayer.Value;
+		
+		if (UCompositeLayerPlate* PlateOwner = Cast<UCompositeLayerPlate>(Layer))
+		{
+			if (!InstancedCustomizations.Contains(UCompositeLayerPlate::StaticClass()))
+			{
+				continue;
+			}
+
+			TWeakPtr<IDetailCustomization>* Customization = InstancedCustomizations[UCompositeLayerPlate::StaticClass()].FindByPredicate([PlateOwner](const TWeakPtr<IDetailCustomization>& Customization)
+			{
+				if (TSharedPtr<IDetailCustomization> PinnedCustomization = Customization.Pin())
+				{
+					if (TSharedPtr<FCompositeLayerPlateCustomization> PlateCustomization = StaticCastSharedPtr<FCompositeLayerPlateCustomization>(PinnedCustomization))
+					{
+						return PlateCustomization->IsCustomizingObject(PlateOwner);
+					}
+				}
+
+				return false;
+			});
+
+			if (Customization)
+			{
+				if (TSharedPtr<FCompositeLayerPlateCustomization> PlateCustomization = StaticCastSharedPtr<FCompositeLayerPlateCustomization>(Customization->Pin()))
+				{
+					if (TSharedPtr<SCompositePlatePassPanel> PassPanel = PlateCustomization->GetPassPanelWidget())
+					{
+						PassPanel->SelectPasses(Passes);
+					}
+				}
+			}
+		}
+	}
+}
+
+TArray<UObject*> SCompositeEditorPanel::GetSelectedObjects() const
+{
+	TArray<UObject*> SelectedObjects;
+	if (LayerTree.IsValid())
+	{
+		SelectedObjects.Append(LayerTree->GetSelectedObjects());
+	}
+
+	for (const TTuple<UClass*, TArray<TWeakPtr<IDetailCustomization>>>& CustomizationsForClass : InstancedCustomizations)
+	{
+		if (CustomizationsForClass.Key == UCompositeLayerPlate::StaticClass())
+		{
+			for (const TWeakPtr<IDetailCustomization>& Customization : CustomizationsForClass.Value)
+			{
+				if (!Customization.IsValid())
+				{
+					continue;
+				}
+
+				if (TSharedPtr<FCompositeLayerPlateCustomization> PlateCustomization = StaticCastSharedPtr<FCompositeLayerPlateCustomization>(Customization.Pin()))
+				{
+					if (TSharedPtr<SCompositePlatePassPanel> PassPanel = PlateCustomization->GetPassPanelWidget())
+					{
+						SelectedObjects.Append(PassPanel->GetSelectedPasses());
+					}
+				}
+			}
+		}
+	}
+	
+	return SelectedObjects;
+}
+
 void SCompositeEditorPanel::OnLayerSelectionChanged(const TArray<UObject*>& SelectedLayers)
 {
 	if (DetailsView.IsValid())
 	{
+		// Clear all stored customizations before changing the details view's objects, as it will recreate any customizations it needs
+		InstancedCustomizations.Empty();
+		
 		DetailsView->SetObjects(SelectedLayers);
 	}
+
+	GetOnSelectionChanged().Broadcast();
 }
 
 float SCompositeEditorPanel::GetLayerTreeMinHeight() const
diff --git a/Engine/Plugins/Compositing/Composite/Source/CompositeEditor/Private/UI/SCompositeEditorPanel.h b/Engine/Plugins/Compositing/Composite/Source/CompositeEditor/Private/UI/SCompositeEditorPanel.h
--- a/Engine/Plugins/Compositing/Composite/Source/CompositeEditor/Private/UI/SCompositeEditorPanel.h
+++ b/Engine/Plugins/Compositing/Composite/Source/CompositeEditor/Private/UI/SCompositeEditorPanel.h
@@ -4,6 +4,8 @@
 
 #include "Widgets/SCompoundWidget.h"
 
+class UCompositePassBase;
+class IDetailCustomization;
 class ACompositeActor;
 class IDetailsView;
 class SCompositePanelLayerTree;
@@ -14,11 +16,22 @@ class SCompositePanelLayerTree;
 class SCompositeEditorPanel : public SCompoundWidget
 {
 public:
+	DECLARE_MULTICAST_DELEGATE(FOnSelectionChanged)
+	
 	/** Registers the tab spawner for the panel with the editor */
 	static void RegisterTabSpawner();
 
 	/** Unregisters the tab spawner for the panel with the editor */
 	static void UnregisterTabSpawner();
+
+	/** Updates the selected composite actors, layers, and passes in the active Composure panel, if one is open */
+	static void UpdateActivePanelSelection(const TArray<UObject*>& InSelection);
+
+	/** Gets the currently selected composite actors, layers, and passes in the active Composure panel, if one is open */
+	static TArray<UObject*> GetActivePanelSelection();
+
+	/** Static delegate that is raised whenever the selection is changed in the active Composure panel */
+	static FOnSelectionChanged& GetOnSelectionChanged() { static FOnSelectionChanged Instance; return Instance; }
 	
 public:
 	SLATE_BEGIN_ARGS(SCompositeEditorPanel) { }
@@ -30,11 +43,36 @@ public:
 	void SelectCompositeActors(const TArray<TWeakObjectPtr<ACompositeActor>>& InCompositeActors);
 	
 private:
+	/** Selects any composite actors, layers, and passes in the specified list of objects */
+	void SelectObjects(const TArray<UObject*>& InSelection);
+
+	/** Gets the actively selected composite actors, layers, and passes in this panel */
+	TArray<UObject*> GetSelectedObjects() const;
+
 	/** Raised when the selected layers in the tree view have changed */
 	void OnLayerSelectionChanged(const TArray<UObject*>& SelectedLayers);
 
 	/** Gets the minimum height to display the full layer tree */
 	float GetLayerTreeMinHeight() const;
+
+	/**
+	 * Generic callback to create a customization of the specified type for the specified object type being customized.
+	 * A weak reference to the customization will be stored in the InstancedCustomizations map keyed under the type being customized
+	 */
+	template<typename TClass, typename TCustomization>
+	TSharedRef<IDetailCustomization> CreateCustomization()
+	{
+		TSharedRef<IDetailCustomization> Customization = TCustomization::MakeInstance();
+
+		if (!InstancedCustomizations.Contains(TClass::StaticClass()))
+		{
+			InstancedCustomizations.Add(TClass::StaticClass(), { });
+		}
+
+		InstancedCustomizations[TClass::StaticClass()].Add(Customization);
+	
+		return Customization;
+	}
 	
 private:
 	/** The tree view widget that displays the composite layers */
@@ -42,4 +80,7 @@ private:
 
 	/** The details view that displays the properties of the selected layers and actors */
 	TSharedPtr<IDetailsView> DetailsView;
+
+	/** A list of pointers to details customizations in use by the details panel, indexed by class being customized */
+	TMap<UClass*, TArray<TWeakPtr<IDetailCustomization>>> InstancedCustomizations;
 };
diff --git a/Engine/Plugins/Compositing/Composite/Source/CompositeEditor/Private/UI/SCompositePanelLayerTree.cpp b/Engine/Plugins/Compositing/Composite/Source/CompositeEditor/Private/UI/SCompositePanelLayerTree.cpp
--- a/Engine/Plugins/Compositing/Composite/Source/CompositeEditor/Private/UI/SCompositePanelLayerTree.cpp
+++ b/Engine/Plugins/Compositing/Composite/Source/CompositeEditor/Private/UI/SCompositePanelLayerTree.cpp
@@ -1111,6 +1111,97 @@ void SCompositePanelLayerTree::SelectCompositeActors(const TArray<TWeakObjectPtr
 	}
 }
 
+void SCompositePanelLayerTree::SelectObjects(const TArray<UObject*>& InObjects)
+{
+	if (!TreeView.IsValid())
+	{
+		return;
+	}
+
+	TreeView->ClearSelection();
+
+	bool bSelectionChanged = false;
+	
+	// Silence the normal invoke of selection changes in the tree. A single, bulk selection change will be invoked after all objects are selected
+	TGuardValue<bool> SilenceSelectionChangesGuard(bSilenceSelectionChanges, true);
+	for (const UObject* Object : InObjects)
+	{
+		if (const ACompositeActor* CompositeActor = Cast<ACompositeActor>(Object))
+		{
+			FCompositeActorTreeItemPtr* TreeItem = CompositeActorTreeItems.FindByPredicate([&CompositeActor](const FCompositeActorTreeItemPtr& InTreeItem)
+			{
+				return InTreeItem.IsValid() && InTreeItem->CompositeActor == CompositeActor;
+			});
+
+			if (TreeItem)
+			{
+				TreeView->SetItemSelection(*TreeItem, true);
+				bSelectionChanged = true;
+			}
+		}
+		else if (const UCompositeLayerBase* CompositeLayer = Cast<UCompositeLayerBase>(Object))
+		{
+			FCompositeActorTreeItemPtr* TreeItem = nullptr;
+			for (const FCompositeActorTreeItemPtr& ActorTreeItem : CompositeActorTreeItems)
+			{
+				if (!ActorTreeItem.IsValid())
+				{
+					continue;
+				}
+
+				TreeItem = ActorTreeItem->Children.FindByPredicate([&CompositeLayer](const FCompositeActorTreeItemPtr& InTreeItem)
+				{
+					return InTreeItem.IsValid() && InTreeItem->HasValidCompositeLayer() && InTreeItem->GetCompositeLayer() == CompositeLayer;
+				});
+
+				if (TreeItem)
+				{
+					break;
+				}
+			}
+
+			if (TreeItem)
+			{
+				TreeView->SetItemSelection(*TreeItem, true);
+				bSelectionChanged = true;
+			}
+		}
+	}
+
+	if (bSelectionChanged)
+	{
+		InvokeLayerSelectionChanged();
+	}
+}
+
+TArray<UObject*> SCompositePanelLayerTree::GetSelectedObjects()
+{
+	TArray<UObject*> SelectedObjects;
+
+	if (TreeView.IsValid())
+	{
+		TArray<FCompositeActorTreeItemPtr> SelectedTreeItems = TreeView->GetSelectedItems();
+		for (const FCompositeActorTreeItemPtr& TreeItem : SelectedTreeItems)
+		{
+			if (!TreeItem.IsValid())
+			{
+				continue;
+			}
+
+			if (TreeItem->HasValidCompositeLayer())
+			{
+				SelectedObjects.Add(TreeItem->GetCompositeLayer());
+			}
+			else if (TreeItem->CompositeActor.IsValid())
+			{
+				SelectedObjects.Add(TreeItem->CompositeActor.Get());
+			}
+		}
+	}
+	
+	return SelectedObjects;
+}
+
 float SCompositePanelLayerTree::GetMinimumHeight() const
 {
 	// Calculate the number of rows the tree view should be displaying
@@ -1724,11 +1815,14 @@ void SCompositePanelLayerTree::OnGlobalEnabledCheckStateChanged(ECheckBoxState C
 
 void SCompositePanelLayerTree::OnLayerSelectionChanged(FCompositeActorTreeItemPtr InTreeItem, ESelectInfo::Type SelectInfo)
 {
-	if (bSilenceSelectionChanges)
+	if (!bSilenceSelectionChanges)
 	{
-		return;
+		InvokeLayerSelectionChanged();
 	}
-	
+}
+
+void SCompositePanelLayerTree::InvokeLayerSelectionChanged()
+{
 	TArray<FCompositeActorTreeItemPtr> SelectedItems = TreeView.IsValid() ? TreeView->GetSelectedItems() : TArray<FCompositeActorTreeItemPtr>();
 	
 	TArray<UObject*> SelectedLayers;
diff --git a/Engine/Plugins/Compositing/Composite/Source/CompositeEditor/Private/UI/SCompositePanelLayerTree.h b/Engine/Plugins/Compositing/Composite/Source/CompositeEditor/Private/UI/SCompositePanelLayerTree.h
--- a/Engine/Plugins/Compositing/Composite/Source/CompositeEditor/Private/UI/SCompositePanelLayerTree.h
+++ b/Engine/Plugins/Compositing/Composite/Source/CompositeEditor/Private/UI/SCompositePanelLayerTree.h
@@ -70,6 +70,12 @@ public:
 
 	/** Selects the specified composite actors in the layer tree */
 	void SelectCompositeActors(const TArray<TWeakObjectPtr<ACompositeActor>>& InCompositeActors);
+
+	/** Selects all matching composite actors and layers in the tree from the specified list of objects */
+	void SelectObjects(const TArray<UObject*>& InObjects);
+
+	/** Gets a list of the currently selected composite actors and layers in the tree */
+	TArray<UObject*> GetSelectedObjects();
 	
 	/** Computes the minimum height necessary to display the layer tree with all layers visible */
 	float GetMinimumHeight() const;
@@ -132,6 +138,9 @@ private:
 	/** Callback that is raised when the selected items in the tree view change */
 	void OnLayerSelectionChanged(FCompositeActorTreeItemPtr InTreeItem, ESelectInfo::Type SelectInfo);
 
+	/** Invokes the OnSelectionChanged delegate with the appropriate selected layers */
+	void InvokeLayerSelectionChanged();
+	
 	/** Callback that is raised when the tree view filter is changed, applies the new filter to the composite actor tree items */
 	void OnFilterChanged();
 
diff --git a/Engine/Plugins/Compositing/Composite/Source/CompositeEditor/Private/UI/SCompositePassTree.cpp b/Engine/Plugins/Compositing/Composite/Source/CompositeEditor/Private/UI/SCompositePassTree.cpp
--- a/Engine/Plugins/Compositing/Composite/Source/CompositeEditor/Private/UI/SCompositePassTree.cpp
+++ b/Engine/Plugins/Compositing/Composite/Source/CompositeEditor/Private/UI/SCompositePassTree.cpp
@@ -1056,6 +1056,61 @@ FReply SCompositePassTree::OnKeyDown(const FGeometry& MyGeometry, const FKeyEven
 	return FReply::Unhandled();
 }
 
+void SCompositePassTree::SelectPasses(const TArray<UCompositePassBase*>& InPasses)
+{
+	if (!TreeView.IsValid())
+	{
+		return;
+	}
+	
+	TreeView->ClearSelection();
+
+	for (UCompositePassBase* Pass : InPasses)
+	{
+		for (const FPassTreeItemPtr& TreeItem : FilteredPassTreeItems)
+		{
+			if (!TreeItem.IsValid())
+			{
+				continue;
+			}
+			
+			FPassTreeItemPtr* TreeItemToSelect = TreeItem->Children.FindByPredicate([Pass](const FPassTreeItemPtr& ChildItem)
+			{
+				return ChildItem.IsValid() && ChildItem->HasValidPassIndex() && ChildItem->GetPass() == Pass;
+			});
+
+			if (TreeItemToSelect)
+			{
+				TreeView->SetItemSelection(*TreeItemToSelect, true);
+			}
+		}
+	}
+}
+
+TArray<UCompositePassBase*> SCompositePassTree::GetSelectedPasses() const
+{
+	TArray<UCompositePassBase*> SelectedPasses;
+	
+	if (TreeView.IsValid())
+	{
+		TArray<FPassTreeItemPtr> SelectedItems = TreeView->GetSelectedItems();
+		for (const FPassTreeItemPtr& SelectedItem : SelectedItems)
+		{
+			if (!SelectedItem.IsValid())
+			{
+				continue;
+			}
+
+			if (UCompositePassBase* Pass = SelectedItem->GetPass())
+			{
+				SelectedPasses.Add(Pass);
+			}
+		}
+	}
+
+	return SelectedPasses;
+}
+
 TArray<TObjectPtr<UCompositePassBase>>& SCompositePassTree::GetPassList(TWeakObjectPtr<UCompositeLayerPlate> InPlate, EPassType InPassType)
 {
 	check(InPlate.IsValid());
diff --git a/Engine/Plugins/Compositing/Composite/Source/CompositeEditor/Private/UI/SCompositePassTree.h b/Engine/Plugins/Compositing/Composite/Source/CompositeEditor/Private/UI/SCompositePassTree.h
--- a/Engine/Plugins/Compositing/Composite/Source/CompositeEditor/Private/UI/SCompositePassTree.h
+++ b/Engine/Plugins/Compositing/Composite/Source/CompositeEditor/Private/UI/SCompositePassTree.h
@@ -67,6 +67,12 @@ public:
 	// SWidget interface
 	virtual FReply OnKeyDown(const FGeometry& MyGeometry, const FKeyEvent& InKeyEvent) override;
 	// End of SWidget interface
+
+	/** Selects the specified passes in the pass tree */
+	void SelectPasses(const TArray<UCompositePassBase*>& InPasses);
+
+	/** Gets a list of currently selected passes in the pass tree */
+	TArray<UCompositePassBase*> GetSelectedPasses() const;
 	
 private:
 	/** Gets a reference to the pass list that corresponds to the specified pass type in the plate. The plate pointer must be valid to call properly */
diff --git a/Engine/Plugins/Compositing/Composite/Source/CompositeEditor/Private/UI/SCompositePlatePassPanel.cpp b/Engine/Plugins/Compositing/Composite/Source/CompositeEditor/Private/UI/SCompositePlatePassPanel.cpp
--- a/Engine/Plugins/Compositing/Composite/Source/CompositeEditor/Private/UI/SCompositePlatePassPanel.cpp
+++ b/Engine/Plugins/Compositing/Composite/Source/CompositeEditor/Private/UI/SCompositePlatePassPanel.cpp
@@ -6,6 +6,7 @@
 #include "DetailsViewArgs.h"
 #include "PropertyEditorModule.h"
 #include "PropertyPath.h"
+#include "SCompositeEditorPanel.h"
 #include "SCompositePassTree.h"
 #include "Customizations/CompositePassBaseCustomization.h"
 #include "Layers/CompositeLayerPlate.h"
@@ -79,6 +80,24 @@ void SCompositePlatePassPanel::Tick(const FGeometry& AllottedGeometry, const dou
 	}
 }
 
+void SCompositePlatePassPanel::SelectPasses(TArray<UCompositePassBase*>& InPasses)
+{
+	if (TreeView.IsValid())
+	{
+		TreeView->SelectPasses(InPasses);
+	}
+}
+
+TArray<UCompositePassBase*> SCompositePlatePassPanel::GetSelectedPasses() const
+{
+	if (TreeView)
+	{
+		return TreeView->GetSelectedPasses();
+	}
+
+	return { };
+}
+
 void SCompositePlatePassPanel::OnTreeViewSelectionChanged(const TArray<UObject*>& Objects)
 {
 	if (!DetailsView.IsValid())
@@ -87,5 +106,7 @@ void SCompositePlatePassPanel::OnTreeViewSelectionChanged(const TArray<UObject*>
 	}
 
 	DetailsView->SetObjects(Objects);
+
+	SCompositeEditorPanel::GetOnSelectionChanged().Broadcast();
 }
 
diff --git a/Engine/Plugins/Compositing/Composite/Source/CompositeEditor/Private/UI/SCompositePlatePassPanel.h b/Engine/Plugins/Compositing/Composite/Source/CompositeEditor/Private/UI/SCompositePlatePassPanel.h
--- a/Engine/Plugins/Compositing/Composite/Source/CompositeEditor/Private/UI/SCompositePlatePassPanel.h
+++ b/Engine/Plugins/Compositing/Composite/Source/CompositeEditor/Private/UI/SCompositePlatePassPanel.h
@@ -24,6 +24,12 @@ public:
 	// SWidget interface
 	virtual void Tick(const FGeometry& AllottedGeometry, const double InCurrentTime, const float InDeltaTime) override;
 	// End of SWidget interface
+
+	/** Selects the specified passes in the pass tree */
+	void SelectPasses(TArray<UCompositePassBase*>& InPasses);
+
+	/** Gets a list of currently selected passes in the pass tree */
+	TArray<UCompositePassBase*> GetSelectedPasses() const;
 	
 private:
 	/** Raised when the selected items in the tree view have changed */
diff --git a/Engine/Plugins/Editor/ColorGrading/Source/ColorGradingEditor/Private/ColorGradingEditorModule.cpp b/Engine/Plugins/Editor/ColorGrading/Source/ColorGradingEditor/Private/ColorGradingEditorModule.cpp
--- a/Engine/Plugins/Editor/ColorGrading/Source/ColorGradingEditor/Private/ColorGradingEditorModule.cpp
+++ b/Engine/Plugins/Editor/ColorGrading/Source/ColorGradingEditor/Private/ColorGradingEditorModule.cpp
@@ -11,6 +11,7 @@
 
 #include "Camera/CameraActor.h"
 #include "CineCameraActor.h"
+#include "ColorGradingMixerSelectionInterface.h"
 #include "Engine/PostProcessVolume.h"
 #include "Framework/Docking/LayoutExtender.h"
 #include "LevelEditor.h"
@@ -72,7 +73,8 @@ TSharedRef<SDockTab> FColorGradingEditorModule::SpawnMainPanelTab(const FSpawnTa
 {
 	if (!MainPanel.IsValid())
 	{
-		SAssignNew(MainPanel, SColorGradingPanel);
+		SAssignNew(MainPanel, SColorGradingPanel)
+		.SelectionInterface(MakeShared<FColorGradingMixerSelectionInterface>());
 	}
 
 	const TSharedRef<SDockTab> DockTab = SNew(SDockTab).TabRole(ETabRole::NomadTab);
diff --git a/Engine/Plugins/Editor/ColorGrading/Source/ColorGradingEditor/Private/ColorGradingMixerObjectFilter.cpp b/Engine/Plugins/Editor/ColorGrading/Source/ColorGradingEditor/Private/ColorGradingMixerObjectFilter.cpp
--- a/Engine/Plugins/Editor/ColorGrading/Source/ColorGradingEditor/Private/ColorGradingMixerObjectFilter.cpp
+++ b/Engine/Plugins/Editor/ColorGrading/Source/ColorGradingEditor/Private/ColorGradingMixerObjectFilter.cpp
@@ -34,6 +34,19 @@ TArray<AActor*> UColorGradingMixerObjectFilter::FindAssociatedActors(AActor* InA
 	return {};
 }
 
+TArray<FObjectMixerEditorActorSubObject> UColorGradingMixerObjectFilter::GetActorSubObjects(AActor* InActor) const
+{
+	if (InActor)
+	{
+		if (const IColorGradingMixerObjectHierarchyConfig* Config = FColorGradingMixerObjectFilterRegistry::GetClassHierarchyConfig(InActor->GetClass()))
+		{
+			return Config->GetActorSubObjects(InActor);
+		}
+	}
+
+	return {};
+}
+
 
 bool UColorGradingMixerObjectFilter::IsActorAssociated(AActor* Actor, AActor* AssociatedActor) const
 {
diff --git a/Engine/Plugins/Editor/ColorGrading/Source/ColorGradingEditor/Private/ColorGradingMixerObjectFilter.h b/Engine/Plugins/Editor/ColorGrading/Source/ColorGradingEditor/Private/ColorGradingMixerObjectFilter.h
--- a/Engine/Plugins/Editor/ColorGrading/Source/ColorGradingEditor/Private/ColorGradingMixerObjectFilter.h
+++ b/Engine/Plugins/Editor/ColorGrading/Source/ColorGradingEditor/Private/ColorGradingMixerObjectFilter.h
@@ -17,6 +17,7 @@ public:
 	UE_API virtual TSet<UClass*> GetObjectClassesToFilter() const override;
 	UE_API virtual TSet<TSubclassOf<AActor>> GetObjectClassesToPlace() const override;
 	UE_API virtual TArray<AActor*> FindAssociatedActors(AActor* InActor) const override;
+	UE_API virtual TArray<FObjectMixerEditorActorSubObject> GetActorSubObjects(AActor* InActor) const override;
 	UE_API virtual bool IsActorAssociated(AActor* Actor, AActor* AssociatedActor) const override;
 	UE_API virtual bool HasCustomDropHandling(const ISceneOutlinerTreeItem& DropTarget) const override;
 	UE_API virtual FSceneOutlinerDragValidationInfo ValidateDrop(const ISceneOutlinerTreeItem& DropTarget, const FSceneOutlinerDragDropPayload& Payload) const override;
diff --git a/Engine/Plugins/Editor/ColorGrading/Source/ColorGradingEditor/Private/ColorGradingMixerObjectFilterRegistry.cpp b/Engine/Plugins/Editor/ColorGrading/Source/ColorGradingEditor/Private/ColorGradingMixerObjectFilterRegistry.cpp
--- a/Engine/Plugins/Editor/ColorGrading/Source/ColorGradingEditor/Private/ColorGradingMixerObjectFilterRegistry.cpp
+++ b/Engine/Plugins/Editor/ColorGrading/Source/ColorGradingEditor/Private/ColorGradingMixerObjectFilterRegistry.cpp
@@ -26,4 +26,19 @@ const IColorGradingMixerObjectHierarchyConfig* FColorGradingMixerObjectFilterReg
 	// No config found. Cache nullptr so we don't have to look this up again
 	HierarchyConfigInstances.Add(Class, nullptr);
 	return nullptr;
-}
\ No newline at end of file
+}
+
+TArray<TSharedPtr<IColorGradingMixerSelectionHandler>> FColorGradingMixerObjectFilterRegistry::CreateSelectionHandlers()
+{
+	TArray<TSharedPtr<IColorGradingMixerSelectionHandler>> Handlers;
+
+	for (const FGetSelectionHandler& CreateHandlerDelegate : SelectionHandlers)
+	{
+		if (CreateHandlerDelegate.IsBound())
+		{
+			Handlers.Add(CreateHandlerDelegate.Execute());
+		}
+	}
+
+	return Handlers;
+}
diff --git a/Engine/Plugins/Editor/ColorGrading/Source/ColorGradingEditor/Private/ColorGradingMixerSelectionInterface.cpp b/Engine/Plugins/Editor/ColorGrading/Source/ColorGradingEditor/Private/ColorGradingMixerSelectionInterface.cpp
new file mode 100644
--- /dev/null
+++ b/Engine/Plugins/Editor/ColorGrading/Source/ColorGradingEditor/Private/ColorGradingMixerSelectionInterface.cpp
@@ -0,0 +1,48 @@
+﻿// Copyright Epic Games, Inc. All Rights Reserved.
+
+#include "ColorGradingMixerSelectionInterface.h"
+
+#include "ColorGradingMixerObjectFilterRegistry.h"
+
+FColorGradingMixerSelectionInterface::FColorGradingMixerSelectionInterface()
+{
+	SelectionHandlers = FColorGradingMixerObjectFilterRegistry::CreateSelectionHandlers();
+	for (const TSharedPtr<IColorGradingMixerSelectionHandler>& Handler : SelectionHandlers)
+	{
+		Handler->OnSelectionChanged.BindRaw(this, &FColorGradingMixerSelectionInterface::OnEditorSelectionChanged);
+	}
+}
+
+void FColorGradingMixerSelectionInterface::SelectSubObjects(const TArray<UObject*>& InSelectedSubObjects, bool bShouldSelect, bool bSelectEvenIfHidden)
+{
+	for (const TSharedPtr<IColorGradingMixerSelectionHandler>& Handler : SelectionHandlers)
+	{
+		TArray<UObject*> Subselection;
+		for (UObject* Object : InSelectedSubObjects)
+		{
+			if (Handler->CanSelectSubObject(Object))
+			{
+				Subselection.Add(Object);
+			}
+		}
+
+		Handler->SelectSubObjects(Subselection, bShouldSelect, bSelectEvenIfHidden);
+	}
+}
+
+TArray<UObject*> FColorGradingMixerSelectionInterface::GetSelectedSubObjects() const
+{
+	TArray<UObject*> SelectedSubObjects;
+	
+	for (const TSharedPtr<IColorGradingMixerSelectionHandler>& Handler : SelectionHandlers)
+	{
+		SelectedSubObjects.Append(Handler->GetSelectedSubObjectsInEditor());
+	}
+	
+	return SelectedSubObjects;
+}
+
+void FColorGradingMixerSelectionInterface::OnEditorSelectionChanged()
+{
+	OnSelectionChanged().Broadcast();
+}
diff --git a/Engine/Plugins/Editor/ColorGrading/Source/ColorGradingEditor/Private/ColorGradingMixerSelectionInterface.h b/Engine/Plugins/Editor/ColorGrading/Source/ColorGradingEditor/Private/ColorGradingMixerSelectionInterface.h
new file mode 100644
--- /dev/null
+++ b/Engine/Plugins/Editor/ColorGrading/Source/ColorGradingEditor/Private/ColorGradingMixerSelectionInterface.h
@@ -0,0 +1,26 @@
+﻿// Copyright Epic Games, Inc. All Rights Reserved.
+
+#pragma once
+
+#include "SelectionInterface/LevelEditorObjectMixerSelectionInterface.h"
+
+struct IColorGradingMixerSelectionHandler;
+
+/** Custom selection interface for the color grading drawer's mixer to provide hooks for receiving selection events for actor sub-objects */
+class FColorGradingMixerSelectionInterface : public FLevelEditorObjectMixerSelectionInterface
+{
+public:
+	FColorGradingMixerSelectionInterface();
+	
+	//~ Begin IObjectMixerSelectionInterface interface
+	virtual void SelectSubObjects(const TArray<UObject*>& InSelectedSubObjects, bool bShouldSelect, bool bSelectEvenIfHidden) override;
+	virtual TArray<UObject*> GetSelectedSubObjects() const override;
+	//~ End IObjectMixerSelectionInterface interface
+
+private:
+	void OnEditorSelectionChanged();
+	
+private:
+	/** A list of intantiated selection handlers that will process selection of actor sub-objects */
+	TArray<TSharedPtr<IColorGradingMixerSelectionHandler>> SelectionHandlers;
+};
diff --git a/Engine/Plugins/Editor/ColorGrading/Source/ColorGradingEditor/Private/SColorGradingPanel.cpp b/Engine/Plugins/Editor/ColorGrading/Source/ColorGradingEditor/Private/SColorGradingPanel.cpp
--- a/Engine/Plugins/Editor/ColorGrading/Source/ColorGradingEditor/Private/SColorGradingPanel.cpp
+++ b/Engine/Plugins/Editor/ColorGrading/Source/ColorGradingEditor/Private/SColorGradingPanel.cpp
@@ -22,6 +22,7 @@
 #include "Views/List/ObjectMixerEditorList.h"
 #include "Views/List/RowTypes/ObjectMixerEditorListRowActor.h"
 #include "Views/List/SObjectMixerEditorList.h"
+#include "Views/List/RowTypes/ObjectMixerEditorListRowUObject.h"
 #include "Widgets/Input/SButton.h"
 #include "Widgets/Layout/SBox.h"
 #include "Widgets/Layout/SExpandableArea.h"
@@ -240,7 +241,6 @@ void SColorGradingPanel::SetPanelState(const FColorGradingPanelState& InPanelSta
 		if (Item.IsValid())
 		{
 			ItemsToSelect.Add(Item);
-			break;
 		}
 	}
 
@@ -593,6 +593,16 @@ void SColorGradingPanel::GetSelectedAndControlledObjects(const TArray<FSceneOutl
 				continue;
 			}
 		}
+
+		if (const FObjectMixerEditorListRowUObject* ObjectTreeItem = TreeItem->CastTo<FObjectMixerEditorListRowUObject>())
+		{
+			if (UObject* Object = ObjectTreeItem->ObjectSoftPtr.Get())
+			{
+				OutControlledObjects.Add(Object);
+				OutSelectedObjects.Add(Object);
+				continue;
+			}
+		}
 	}
 }
 
diff --git a/Engine/Plugins/Editor/ColorGrading/Source/ColorGradingEditor/Public/ColorGradingMixerObjectFilterRegistry.h b/Engine/Plugins/Editor/ColorGrading/Source/ColorGradingEditor/Public/ColorGradingMixerObjectFilterRegistry.h
--- a/Engine/Plugins/Editor/ColorGrading/Source/ColorGradingEditor/Public/ColorGradingMixerObjectFilterRegistry.h
+++ b/Engine/Plugins/Editor/ColorGrading/Source/ColorGradingEditor/Public/ColorGradingMixerObjectFilterRegistry.h
@@ -5,6 +5,7 @@
 #include "GameFramework/Actor.h"
 #include "HAL/Platform.h"
 #include "SceneOutlinerDragDrop.h"
+#include "ObjectFilter/ObjectMixerEditorActorSubObject.h"
 #include "Templates/SubclassOf.h"
 
 #define UE_API COLORGRADINGEDITOR_API
@@ -27,6 +28,9 @@ public:
 	 */
 	virtual bool IsActorAssociated(UObject* ParentObject, AActor* AssociatedActor) const { return false; }
 
+	/** Generate a list of additional sub-objects to include as children in the ObjectMixer hierarchy for the given object */
+	virtual TArray<FObjectMixerEditorActorSubObject> GetActorSubObjects(UObject* ParentObject) const { return {}; }
+
 	/** 
 	 * If true, this has custom handling for drop operations and will override the hierarchy's default behaviour with the ValidateDrop
 	 * and OnDrop functions.
@@ -48,6 +52,29 @@ public:
 
 DECLARE_DELEGATE_RetVal(TSharedRef<IColorGradingMixerObjectHierarchyConfig>, FGetObjectHierarchyConfig);
 
+/** Interface for custom selection handlers that can receive selection events from the color grading drawer's mixer */
+struct IColorGradingMixerSelectionHandler : public TSharedFromThis<IColorGradingMixerSelectionHandler>
+{
+public:
+	DECLARE_DELEGATE(FOnSelectionChanged);
+
+	virtual ~IColorGradingMixerSelectionHandler() = default;
+	
+	/** Returns true if this handler can handle synchronizing the specified object */
+	virtual bool CanSelectSubObject(const UObject* InObject) { return false; }
+	
+	/** Called when the color grading mixer synchronizes its selection to the rest of the editor */
+	virtual void SelectSubObjects(const TArray<UObject*>& InSelectedObjects, bool bShouldSelect, bool bSelectEvenIfHidden) const { }
+
+	/** Called when the color grading mixer synchronizes its selection from selected objects in the editor. Retrieves a list of all sub-objects selected in the editor */
+	virtual TArray<UObject*> GetSelectedSubObjectsInEditor() const { return {}; }
+
+	/** Delegate that can be invoked to cause the color grading mixer to re-synchronize its selection to the editor's selection */
+	FOnSelectionChanged OnSelectionChanged;
+};
+
+DECLARE_DELEGATE_RetVal(TSharedRef<IColorGradingMixerSelectionHandler>, FGetSelectionHandler);
+
 /**
  * Contains functions for controlling which classes can be managed from the Color Grading panel's ObjectMixer-based hierarchy panel
  */
@@ -76,6 +103,15 @@ public:
 		ActorClassesToPlace.Add(Class);
 	}
 
+	/**
+	 * Registers a new selection handler with the color drawer that will receive selection related events from the drawer's mixer
+	 * @param InCreateHandlerDelegate A delegate that is invoked to create the selection handler when needed
+	 */
+	static void RegisterSelectionHandler(FGetSelectionHandler InCreateHandlerDelegate)
+	{
+		SelectionHandlers.Add(InCreateHandlerDelegate);	
+	}
+	
 	/** Get the set of object classes that can be seen in a Color Grading panel's object list */
 	static const TSet<UClass*>& GetObjectClassesToFilter() { return ObjectClassesToFilter; }
 
@@ -85,6 +121,9 @@ public:
 	/** Get the hierarchy configuration for a class, or null if none was provided */
 	static UE_API const IColorGradingMixerObjectHierarchyConfig* GetClassHierarchyConfig(TSubclassOf<AActor> Class);
 
+	/** Instantiates a list of all selection handlers that have been registered to the color grading drawer */
+	static TArray<TSharedPtr<IColorGradingMixerSelectionHandler>> CreateSelectionHandlers();
+	
 private:
 	/** Set of classes that can be seen in the object panel */
 	UE_API inline static TSet<UClass*> ObjectClassesToFilter;
@@ -97,6 +136,9 @@ private:
 
 	/** Map from object class to hierarchy configuration instance. If null, none exist for the class */
 	UE_API inline static TMap<TWeakObjectPtr<UClass>, TSharedPtr<IColorGradingMixerObjectHierarchyConfig>> HierarchyConfigInstances;
+
+	/** List of handlers that process custom selection synchronization for specific color gradable objects */
+	UE_API inline static TArray<FGetSelectionHandler> SelectionHandlers;
 };
 
 #undef UE_API
diff --git a/Engine/Plugins/Editor/ObjectMixer/ObjectMixer/Source/ObjectMixer/Private/ObjectFilter/ObjectMixerEditorObjectFilter.cpp b/Engine/Plugins/Editor/ObjectMixer/ObjectMixer/Source/ObjectMixer/Private/ObjectFilter/ObjectMixerEditorObjectFilter.cpp
--- a/Engine/Plugins/Editor/ObjectMixer/ObjectMixer/Source/ObjectMixer/Private/ObjectFilter/ObjectMixerEditorObjectFilter.cpp
+++ b/Engine/Plugins/Editor/ObjectMixer/ObjectMixer/Source/ObjectMixer/Private/ObjectFilter/ObjectMixerEditorObjectFilter.cpp
@@ -72,6 +72,11 @@ bool UObjectMixerObjectFilter::IsActorAssociated(AActor* Actor, AActor* Associat
 	return false;
 }
 
+TArray<FObjectMixerEditorActorSubObject> UObjectMixerObjectFilter::GetActorSubObjects(AActor* InActor) const
+{
+	return {};
+}
+
 bool UObjectMixerObjectFilter::HasCustomDropHandling(const ISceneOutlinerTreeItem& DropTarget) const
 {
 	return false;
diff --git a/Engine/Plugins/Editor/ObjectMixer/ObjectMixer/Source/ObjectMixer/Private/Views/List/Modes/ObjectMixerOutlinerHierarchy.cpp b/Engine/Plugins/Editor/ObjectMixer/ObjectMixer/Source/ObjectMixer/Private/Views/List/Modes/ObjectMixerOutlinerHierarchy.cpp
--- a/Engine/Plugins/Editor/ObjectMixer/ObjectMixer/Source/ObjectMixer/Private/Views/List/Modes/ObjectMixerOutlinerHierarchy.cpp
+++ b/Engine/Plugins/Editor/ObjectMixer/ObjectMixer/Source/ObjectMixer/Private/Views/List/Modes/ObjectMixerOutlinerHierarchy.cpp
@@ -29,6 +29,7 @@
 #include "LevelInstance/LevelInstanceSubsystem.h"
 #include "Modules/ModuleManager.h"
 #include "SceneOutlinerFwd.h"
+#include "Views/List/RowTypes/ObjectMixerEditorListRowContainer.h"
 #include "WorldPartition/ActorDescContainerInstanceCollection.h"
 #include "WorldPartition/ActorDescContainerInstance.h"
 #include "WorldPartition/WorldPartitionActorDescInstance.h"
@@ -426,6 +427,99 @@ FSceneOutlinerTreeItemPtr FObjectMixerOutlinerHierarchy::FindOrCreateParentItem(
 			}
 		}
 	}
+	else if (const FObjectMixerEditorListRowUObject* ObjectItem = Item.CastTo<FObjectMixerEditorListRowUObject>())
+	{
+		if (UObject* Object = ObjectItem->ObjectSoftPtr.Get())
+		{
+			UObject* ParentObject = ObjectItem->bIncludeObjectStack ? Object->GetOuter() : Object->GetTypedOuter<AActor>();
+			if (ParentObject)
+			{
+				if (ObjectItem->ContainerProperty != NAME_None)
+				{
+					FSceneOutlinerTreeItemID ContainerRowId = HashCombineFast(GetTypeHash(ParentObject), GetTypeHash(ObjectItem->ContainerProperty));
+					if (const FSceneOutlinerTreeItemPtr* ParentItem = Items.Find(ContainerRowId))
+					{
+						return *ParentItem;
+					}
+					else if (bCreate)
+					{
+						if (FSceneOutlinerTreeItemPtr ContainerItem =
+							Mode->CreateItemFor<FObjectMixerEditorListRowContainer>(
+								FObjectMixerEditorListRowContainer(ParentObject, ObjectItem->ContainerProperty, GetCastedMode()->GetSceneOutliner()), true))
+						{
+							return ContainerItem;
+						}
+					}
+
+					return nullptr;
+				}
+				
+				if (const FSceneOutlinerTreeItemPtr* ParentItem = Items.Find(ParentObject))
+				{
+					return *ParentItem;
+				}
+				else if (bCreate)
+				{
+					if (ParentObject->IsA<AActor>())
+					{
+						if (FSceneOutlinerTreeItemPtr ActorItem =
+								Mode->CreateItemFor<FObjectMixerEditorListRowActor>(
+									FObjectMixerEditorListRowActor(Cast<AActor>(ParentObject), GetCastedMode()->GetSceneOutliner()), true))
+						{
+							return ActorItem;
+						}
+					}
+					else
+					{
+						if (TSharedPtr<FObjectMixerEditorListRowUObject> ParentObjectItem =
+							StaticCastSharedPtr<FObjectMixerEditorListRowUObject>(Mode->CreateItemFor<FObjectMixerEditorListRowUObject>(
+								FObjectMixerEditorListRowUObject(ParentObject, GetCastedMode()->GetSceneOutliner()), true)))
+						{
+							ParentObjectItem->bIncludeObjectStack = ObjectItem->bIncludeObjectStack;
+							return ParentObjectItem;
+						}
+					}
+				}
+					
+				return nullptr;
+			}
+		}
+	}
+	else if (const FObjectMixerEditorListRowContainer* ContainerItem = Item.CastTo<FObjectMixerEditorListRowContainer>())
+	{
+		if (UObject* PropertyOwner = ContainerItem->GetPropertyOwner())
+		{
+			if (const FSceneOutlinerTreeItemPtr* ParentItem = Items.Find(PropertyOwner))
+			{
+				return *ParentItem;
+			}
+			else if (bCreate)
+			{
+				if (PropertyOwner->IsA<AActor>())
+				{
+					if (FSceneOutlinerTreeItemPtr ActorItem =
+							Mode->CreateItemFor<FObjectMixerEditorListRowActor>(
+								FObjectMixerEditorListRowActor(Cast<AActor>(PropertyOwner), GetCastedMode()->GetSceneOutliner()), true))
+					{
+						return ActorItem;
+					}
+				}
+				else
+				{
+					if (TSharedPtr<FObjectMixerEditorListRowUObject> ParentObjectItem =
+						StaticCastSharedPtr<FObjectMixerEditorListRowUObject>(Mode->CreateItemFor<FObjectMixerEditorListRowUObject>(
+							FObjectMixerEditorListRowUObject(PropertyOwner, GetCastedMode()->GetSceneOutliner()), true)))
+					{
+						// Assume that, since a container item exists, that it is part of a full sub-object stack
+						ParentObjectItem->bIncludeObjectStack = true;
+						return ParentObjectItem;
+					}
+				}
+			}
+					
+			return nullptr;
+		}
+	}
 
 	// If we get here return world item
 	if (const FSceneOutlinerTreeItemPtr* ParentItem = Items.Find(RepresentingWorld.Get()))
@@ -625,6 +719,76 @@ bool FObjectMixerOutlinerHierarchy::CreateAssociatedActorItems(AActor* Actor, TA
 	return bAnyItemsBeforeFiltering;
 }
 
+bool FObjectMixerOutlinerHierarchy::CreateSubObjectItems(AActor* Actor, TArray<FSceneOutlinerTreeItemPtr>& OutItems) const
+{
+	bool bAnyItemsBeforeFiltering = false;
+	
+	check(Actor);
+	if (const UObjectMixerObjectFilter* Filter = GetObjectFilter())
+	{
+		const TArray<FObjectMixerEditorActorSubObject> SubObjects = Filter->GetActorSubObjects(Actor);
+
+		for (const FObjectMixerEditorActorSubObject& SubObject : SubObjects)
+		{
+			if (TStrongObjectPtr<UObject> Object = SubObject.Object.Pin())
+			{
+				bAnyItemsBeforeFiltering = true;
+
+				if (SubObject.bIncludeObjectStack)
+				{
+					bool bIsValidSubObjectStack = false;
+					TArray<UObject*> SubObjectStack;
+					UObject* Outer = Object.Get()->GetOuter();
+					while (Outer != Actor && Outer != nullptr)
+					{
+						SubObjectStack.Add(Outer);
+						Outer = Outer->GetOuter();
+						if (Outer == Actor)
+						{
+							bIsValidSubObjectStack = true;
+						}
+					}
+
+					if (bIsValidSubObjectStack)
+					{
+						for (int32 Index = 0; Index < SubObjectStack.Num(); ++Index)
+						{
+							if (FSceneOutlinerTreeItemPtr ObjectItem =
+								Mode->CreateItemFor<FObjectMixerEditorListRowUObject>(
+									FObjectMixerEditorListRowUObject(SubObjectStack[Index], GetCastedMode()->GetSceneOutliner(), FText::GetEmpty())))
+							{
+								OutItems.Add(ObjectItem);
+							}
+						}
+					}
+				}
+				
+				if (SubObject.ContainerProperty.PropertyName != NAME_None)
+				{
+					if (TSharedPtr<FObjectMixerEditorListRowContainer> ContainerItem =
+						StaticCastSharedPtr<FObjectMixerEditorListRowContainer>(Mode->CreateItemFor<FObjectMixerEditorListRowContainer>(
+							FObjectMixerEditorListRowContainer(Object->GetOuter(), SubObject.ContainerProperty.PropertyName, GetCastedMode()->GetSceneOutliner()))))
+					{
+						ContainerItem->PropertyIconOverride = SubObject.ContainerProperty.PropertyIconOverride;
+						OutItems.Add(ContainerItem);
+					}				
+				}
+				
+				if (TSharedPtr<FObjectMixerEditorListRowUObject> ObjectItem =
+					StaticCastSharedPtr<FObjectMixerEditorListRowUObject>(Mode->CreateItemFor<FObjectMixerEditorListRowUObject>(
+						FObjectMixerEditorListRowUObject(Object.Get(), GetCastedMode()->GetSceneOutliner(), SubObject.DisplayName))))
+				{
+					ObjectItem->ContainerProperty = SubObject.ContainerProperty.PropertyName;
+					ObjectItem->bIncludeObjectStack = SubObject.bIncludeObjectStack;
+					OutItems.Add(ObjectItem);
+				}
+			}
+		}
+	}
+	
+	return bAnyItemsBeforeFiltering;
+}
+
 TArray<FSceneOutlinerTreeItemPtr> FObjectMixerOutlinerHierarchy::ConditionallyCreateActorAndComponentItems(AActor* Actor) const
 {
 	if (!IsValid(Actor))
@@ -641,15 +805,20 @@ TArray<FSceneOutlinerTreeItemPtr> FObjectMixerOutlinerHierarchy::ConditionallyCr
 	TArray<FSceneOutlinerTreeItemPtr> AssociatedActorRows;
 	bWouldReturnAnyChildItemsBeforeFiltering |= CreateAssociatedActorItems(Actor, AssociatedActorRows);
 
+	TArray<FSceneOutlinerTreeItemPtr> SubObjectRows;
+	bWouldReturnAnyChildItemsBeforeFiltering |= CreateSubObjectItems(Actor, SubObjectRows);
+
 	const bool bShouldCreateActorItem =
 		bWouldReturnAnyChildItemsBeforeFiltering || !bShowingOnlyActorWithValidComponents || DoesWorldObjectHaveAcceptableClass(Actor);
 
 	TArray<FSceneOutlinerTreeItemPtr> ReturnValue;
 	if (bShouldCreateActorItem)
 	{
+		// Force the actor creation regardless of whether it passes the current filters, as one of its children did
+		const bool bForce = true;
 		if (const FSceneOutlinerTreeItemPtr ActorItem =
 			Mode->CreateItemFor<FObjectMixerEditorListRowActor>(
-				FObjectMixerEditorListRowActor(Actor, GetCastedMode()->GetSceneOutliner())))
+				FObjectMixerEditorListRowActor(Actor, GetCastedMode()->GetSceneOutliner()), bForce))
 		{
 			const bool bIsHybridModeEnabled = GetCastedMode()->ShouldAllowHybridRows() && GetDefault<UObjectMixerEditorSettings>()->IsHybridRowModeEnabled();
 			if (bIsHybridModeEnabled && ComponentRows.Num() == 1 && AssociatedActorRows.IsEmpty()) // Create hybrid row
@@ -668,6 +837,7 @@ TArray<FSceneOutlinerTreeItemPtr> FObjectMixerOutlinerHierarchy::ConditionallyCr
 			{
 				ReturnValue.Append(ComponentRows);
 				ReturnValue.Append(AssociatedActorRows);
+				ReturnValue.Append(SubObjectRows);
 			}
 
 			// Place ActorItem before components
@@ -726,6 +896,7 @@ void FObjectMixerOutlinerHierarchy::CreateChildren(const FSceneOutlinerTreeItemP
 
 		CreateComponentItems(ParentActor, OutChildren);
 		CreateAssociatedActorItems(ParentActor, OutChildren);
+		CreateSubObjectItems(ParentActor, OutChildren);
 
 		TArray<AActor*> ChildActors;
 
diff --git a/Engine/Plugins/Editor/ObjectMixer/ObjectMixer/Source/ObjectMixer/Private/Views/List/Modes/ObjectMixerOutlinerMode.cpp b/Engine/Plugins/Editor/ObjectMixer/ObjectMixer/Source/ObjectMixer/Private/Views/List/Modes/ObjectMixerOutlinerMode.cpp
--- a/Engine/Plugins/Editor/ObjectMixer/ObjectMixer/Source/ObjectMixer/Private/Views/List/Modes/ObjectMixerOutlinerMode.cpp
+++ b/Engine/Plugins/Editor/ObjectMixer/ObjectMixer/Source/ObjectMixer/Private/Views/List/Modes/ObjectMixerOutlinerMode.cpp
@@ -57,6 +57,8 @@
 #include "Styling/StyleColors.h"
 #include "Subsystems/ActorEditorContextSubsystem.h"
 #include "ToolMenus.h"
+#include "Views/List/RowTypes/ObjectMixerEditorListRowContainer.h"
+#include "Views/List/RowTypes/ObjectMixerEditorListRowUObject.h"
 #include "Widgets/Input/SButton.h"
 #include "Widgets/Input/SComboButton.h"
 #include "WorldPartition/IWorldPartitionEditorModule.h"
@@ -200,6 +202,28 @@ namespace ObjectMixerOutliner
 					}
 				}
 			}
+			else if (FObjectMixerEditorListRowUObject* ObjectItem = ItemPtr->CastTo<FObjectMixerEditorListRowUObject>())
+			{
+				if (UObject* Object = ObjectItem->ObjectSoftPtr.Get())
+				{
+					if (AActor* Actor = Object->GetTypedOuter<AActor>())
+					{
+						ActorPtrOut = Actor;
+						return true;
+					}					
+				}
+			}
+			else if (FObjectMixerEditorListRowContainer* ContainerItem = ItemPtr->CastTo<FObjectMixerEditorListRowContainer>())
+			{
+				if (UObject* PropertyOwner = ContainerItem->GetPropertyOwner())
+				{
+					if (AActor* Actor = PropertyOwner->GetTypedOuter<AActor>())
+					{
+						ActorPtrOut = Actor;
+						return true;
+					}
+				}
+			}
 		}
 
 		return false;
@@ -231,6 +255,54 @@ namespace ObjectMixerOutliner
 		}
 		return false;
 	}
+
+	bool FObjectSelector::operator()(TWeakPtr<ISceneOutlinerTreeItem> Item, UObject*& OutObject)
+	{
+		if (TSharedPtr<ISceneOutlinerTreeItem> ItemPtr = Item.Pin())
+		{
+			if (FObjectMixerEditorListRowActor* ActorItem = ItemPtr->CastTo<FObjectMixerEditorListRowActor>())
+			{
+				if (ActorItem->IsValid())
+				{
+					if (AActor* Actor = ActorItem->Actor.Get())
+					{
+						OutObject = Actor;
+						return true;
+					}
+				}
+			}
+			// If a component is selected, we meant for the owning actor to be selected
+			else if (FComponentTreeItem* ComponentItem = ItemPtr->CastTo<FComponentTreeItem>())
+			{
+				if (ComponentItem->IsValid())
+				{
+					if (UActorComponent* ActorComponent = ComponentItem->Component.Get())
+					{
+						OutObject = ActorComponent;
+						return true;
+					}
+				}
+			}
+			else if (FObjectMixerEditorListRowUObject* ObjectItem = ItemPtr->CastTo<FObjectMixerEditorListRowUObject>())
+			{
+				if (UObject* Object = ObjectItem->ObjectSoftPtr.Get())
+				{
+					OutObject = Object;
+					return true;			
+				}
+			}
+			else if (FObjectMixerEditorListRowContainer* ContainerItem = ItemPtr->CastTo<FObjectMixerEditorListRowContainer>())
+			{
+				if (UObject* PropertyOwner = ContainerItem->GetPropertyOwner())
+				{
+					OutObject = PropertyOwner;
+					return true;
+				}
+			}
+		}
+		
+		return false;
+	}
 }
 
 FObjectMixerOutlinerMode::FObjectMixerOutlinerMode(
@@ -720,6 +792,9 @@ void FObjectMixerOutlinerMode::SynchronizeAllSelectionsToEditor()
 			SelectComponentsInEditor(SelectedComponents, true, true);
 		}
 	}
+
+	TArray<UObject*> SelectedObjects = SceneOutliner->GetSelection().GetData<UObject*>(ObjectMixerOutliner::FObjectSelector());
+	SelectSubObjectsInEditor(SelectedObjects, true, true);	
 }
 
 bool FObjectMixerOutlinerMode::HasActorSelectionChanged(TArray<AActor*>& OutSelectedActors, bool& bOutAreAnyInPIE)
@@ -831,6 +906,17 @@ void FObjectMixerOutlinerMode::SelectComponentsInEditor(const TArray<UActorCompo
 	}
 }
 
+void FObjectMixerOutlinerMode::SelectSubObjectsInEditor(const TArray<UObject*>& InSelectedSubObjects, bool bShouldSelect, bool bSelectEvenIfHidden)
+{
+	if (TSharedPtr<FObjectMixerEditorList> PinnedList = GetListModelPtr().Pin())
+	{
+		if (PinnedList->GetSelectionInterface().IsValid())
+		{
+			PinnedList->GetSelectionInterface()->SelectSubObjects(InSelectedSubObjects, bShouldSelect, bSelectEvenIfHidden);
+		}
+	}
+}
+
 void FObjectMixerOutlinerMode::SelectActorsInMixer(const TArray<AActor*>& InSelectedActors, bool bShouldSelect, bool bSelectEvenIfHidden)
 {
 	TArray<FSceneOutlinerTreeItemPtr> ItemsToSelect;
@@ -905,6 +991,19 @@ TArray<UActorComponent*> FObjectMixerOutlinerMode::GetSelectedComponentsInEditor
 	return TArray<UActorComponent*>();
 }
 
+TArray<UObject*> FObjectMixerOutlinerMode::GetSelectedSubObjectsInEditor() const
+{
+	if (TSharedPtr<FObjectMixerEditorList> PinnedList = GetListModelPtr().Pin())
+	{
+		if (PinnedList->GetSelectionInterface().IsValid())
+		{
+			return PinnedList->GetSelectionInterface()->GetSelectedSubObjects();
+		}
+	}
+
+	return TArray<UObject*>();
+}
+
 bool FObjectMixerOutlinerMode::GetIsObjectOfFilteredClass(const UObject* Object) const
 {
 	if (IsValid(Object))
@@ -1102,6 +1201,28 @@ void FObjectMixerOutlinerMode::Rebuild()
 	SceneOutliner->FullRefresh();
 }
 
+int32 FObjectMixerOutlinerMode::GetTypeSortPriority(const ISceneOutlinerTreeItem& Item) const
+{
+	if (Item.IsA<FWorldTreeItem>())
+	{
+		return EItemSortOrder::World;
+	}
+	else if (Item.IsA<FLevelTreeItem>())
+	{
+		return EItemSortOrder::Level;
+	}
+	else if (Item.IsA<FFolderTreeItem>())
+	{
+		return EItemSortOrder::Folder;
+	}
+	else if (Item.IsA<FActorTreeItem>() || Item.IsA<FComponentTreeItem>() || Item.IsA<FActorDescTreeItem>())
+	{
+		return EItemSortOrder::Actor;
+	}
+
+	return 40;
+}
+
 FText FObjectMixerOutlinerMode::GetStatusText() const
 {
 	if (!RepresentingWorld.IsValid())
@@ -2528,6 +2649,64 @@ void FObjectMixerOutlinerMode::SynchronizeSelectedActorDescs()
 	}
 }
 
+void FObjectMixerOutlinerMode::SynchronizeSubObjectSelection()
+{
+	TArray<UObject*> SelectedSubObjects = GetSelectedSubObjectsInEditor();
+
+	if (SelectedSubObjects.IsEmpty())
+	{
+		return;
+	}
+
+	const FSceneOutlinerItemSelection Selection = SceneOutliner->GetSelection();
+
+	// Deselect any object rows that are not in the current selection
+	Selection.ForEachItem<FObjectMixerEditorListRowUObject>([this, &SelectedSubObjects](FObjectMixerEditorListRowUObject& Item)
+	{
+		if (!Item.ObjectSoftPtr.IsValid() || !SelectedSubObjects.Contains(Item.ObjectSoftPtr.Get()))
+		{
+			SceneOutliner->SetItemSelection(Item.AsShared(), false);
+		}
+	});
+
+	// Deselect any parent actors of the selected sub-objects to alleviate any multi-select issues
+	TArray<AActor*> OwningActors;
+	Algo::Transform(SelectedSubObjects, OwningActors, [](const UObject* Object)
+	{
+		return Object->GetTypedOuter<AActor>();
+	});
+	
+	Selection.ForEachItem<FActorTreeItem>([this, &OwningActors](FActorTreeItem& Item)
+	{
+		if (Item.Actor.IsValid() && OwningActors.Contains(Item.Actor.Get()))
+		{
+			SceneOutliner->SetItemSelection(Item.AsShared(), false);
+		}
+	});
+	
+	bool bSelectorInSelectionSet = false;
+
+	TArray<FSceneOutlinerTreeItemPtr> SelectedItems;
+	for (UObject* SubObject : SelectedSubObjects)
+	{
+		if (FSceneOutlinerTreeItemPtr ObjectItem = SceneOutliner->GetTreeItem(SubObject))
+		{
+			if (!bSelectorInSelectionSet && SceneOutliner->HasSelectorFocus(ObjectItem))
+			{
+				bSelectorInSelectionSet = true;
+			}
+
+			SelectedItems.Add(ObjectItem);
+		}
+	}
+
+	// If NOT bSelectorInSelectionSet then we want to just move the selector to the first selected item.
+	ESelectInfo::Type SelectInfo = bSelectorInSelectionSet ? ESelectInfo::Direct : ESelectInfo::OnMouseClick;
+	SceneOutliner->AddToSelection(SelectedItems, SelectInfo);
+
+	FSceneOutlinerDelegates::Get().SelectionChanged.Broadcast();
+}
+
 FSceneOutlinerDragValidationInfo FObjectMixerOutlinerMode::ValidateDrop(const ISceneOutlinerTreeItem& DropTarget, const FSceneOutlinerDragDropPayload& Payload) const
 {
 	if (const UObjectMixerObjectFilter* Filter = GetListModelPtr().Pin()->GetMainObjectFilterInstance())
@@ -3402,6 +3581,7 @@ void FObjectMixerOutlinerMode::SynchronizeSelection()
 	{
 		SynchronizeComponentAndActorSelection();
 		SynchronizeSelectedActorDescs();
+		SynchronizeSubObjectSelection();
 	}
 	else
 	{
diff --git a/Engine/Plugins/Editor/ObjectMixer/ObjectMixer/Source/ObjectMixer/Private/Views/List/RowTypes/ObjectMixerEditorListRowContainer.cpp b/Engine/Plugins/Editor/ObjectMixer/ObjectMixer/Source/ObjectMixer/Private/Views/List/RowTypes/ObjectMixerEditorListRowContainer.cpp
new file mode 100644
--- /dev/null
+++ b/Engine/Plugins/Editor/ObjectMixer/ObjectMixer/Source/ObjectMixer/Private/Views/List/RowTypes/ObjectMixerEditorListRowContainer.cpp
@@ -0,0 +1,140 @@
+﻿// Copyright Epic Games, Inc. All Rights Reserved.
+
+#include "Views/List/RowTypes/ObjectMixerEditorListRowContainer.h"
+
+#include "ISceneOutliner.h"
+#include "ISceneOutlinerMode.h"
+#include "SceneOutlinerPublicTypes.h"
+#include "Settings/EditorStyleSettings.h"
+#include "Widgets/Images/SImage.h"
+#include "Widgets/Text/SInlineEditableTextBlock.h"
+
+const FSceneOutlinerTreeItemType FObjectMixerEditorListRowContainer::Type(&ISceneOutlinerTreeItem::Type);
+
+FObjectMixerEditorListRowContainer::FObjectMixerEditorListRowContainer(UObject* InPropertyOwner, FName InPropertyName, SSceneOutliner* InSceneOutliner)
+	: ISceneOutlinerTreeItem(Type)
+{
+	PropertyOwner = InPropertyOwner;
+
+	if (PropertyOwner.IsValid())
+	{
+		Property = PropertyOwner->GetClass()->FindPropertyByName(InPropertyName);
+		UniqueId = HashCombineFast(GetTypeHash(InPropertyOwner), GetTypeHash(InPropertyName));
+	}
+}
+
+FString FObjectMixerEditorListRowContainer::GetDisplayString() const
+{
+	if (!IsValid())
+	{
+		return TEXT("");
+	}
+
+	return Property->GetDisplayNameText().ToString();
+}
+
+class SContainerTreeLabel : FSceneOutlinerCommonLabelData, public SCompoundWidget
+{
+public:
+	SLATE_BEGIN_ARGS(SContainerTreeLabel) {}
+	SLATE_END_ARGS()
+
+	void Construct(const FArguments& InArgs, FObjectMixerEditorListRowContainer& ContainerItem, ISceneOutliner& SceneOutliner, const STableRow<FSceneOutlinerTreeItemPtr>& InRow)
+	{
+		WeakSceneOutliner = StaticCastSharedRef<ISceneOutliner>(SceneOutliner.AsShared());
+
+		TreeItemPtr = StaticCastSharedRef<FObjectMixerEditorListRowContainer>(ContainerItem.AsShared());
+
+		const bool bShouldUseMiddleEllipsis = GetDefault<UEditorStyleSettings>()->bEnableMiddleEllipsis;
+
+		ChildSlot
+		[
+			SNew(SHorizontalBox)
+
+			+SHorizontalBox::Slot()
+			.AutoWidth()
+			.VAlign(VAlign_Center)
+			.Padding(FSceneOutlinerDefaultTreeItemMetrics::IconPadding())
+			[
+				SNew(SBox)
+				.WidthOverride(static_cast<float>(FSceneOutlinerDefaultTreeItemMetrics::IconSize()))
+				.HeightOverride(static_cast<float>(FSceneOutlinerDefaultTreeItemMetrics::IconSize()))
+				[
+					SNew(SImage)
+					.Image(this, &SContainerTreeLabel::GetIcon)
+					.ColorAndOpacity(FSlateColor::UseForeground())
+				]
+			]
+
+			+SHorizontalBox::Slot()
+			.FillWidth(1.0f)
+			.VAlign(VAlign_Center)
+			.Padding(0.0f, 0.0f)
+			[
+				SNew(SHorizontalBox)
+				
+				+SHorizontalBox::Slot()
+				.VAlign(VAlign_Center)
+				[
+					SNew(STextBlock)
+					.Text(this, &SContainerTreeLabel::GetDisplayText)
+					.ToolTipText(this, &SContainerTreeLabel::GetTooltipText)
+					.HighlightText(SceneOutliner.GetFilterHighlightText())
+					.ColorAndOpacity(this, &SContainerTreeLabel::GetTextForegroundColor)
+					.OverflowPolicy(bShouldUseMiddleEllipsis ? ETextOverflowPolicy::MiddleEllipsis : TOptional<ETextOverflowPolicy>())
+				]
+			]
+		];
+	}
+
+private:
+	FText GetDisplayText() const
+	{
+		if (const TSharedPtr<FObjectMixerEditorListRowContainer> TreeItem = TreeItemPtr.Pin())
+		{
+			return FText::FromString(TreeItem->GetDisplayString());
+		}
+
+		return FText();
+	}
+
+	FText GetTooltipText() const
+	{
+		if (const TSharedPtr<FObjectMixerEditorListRowContainer> TreeItem = TreeItemPtr.Pin())
+		{
+			return FText::FromString(TreeItem->GetDisplayString());
+		}
+
+		return FText();
+	}
+
+	const FSlateBrush* GetIcon() const
+	{
+		const FSlateBrush* DefaultIcon = nullptr;
+		if (const TSharedPtr<FObjectMixerEditorListRowContainer> TreeItem = TreeItemPtr.Pin())
+		{
+			return TreeItem->PropertyIconOverride.Get(DefaultIcon);
+		}
+		
+		return DefaultIcon;
+	}
+
+	FSlateColor GetTextForegroundColor() const
+	{
+		auto TreeItem = TreeItemPtr.Pin();
+		if (auto BaseColor = FSceneOutlinerCommonLabelData::GetForegroundColor(*TreeItem))
+		{
+			return BaseColor.GetValue();
+		}
+		
+		return FSlateColor::UseForeground();
+	}
+	
+private:
+	TWeakPtr<FObjectMixerEditorListRowContainer> TreeItemPtr;
+};
+
+TSharedRef<SWidget> FObjectMixerEditorListRowContainer::GenerateLabelWidget(ISceneOutliner& Outliner, const STableRow<FSceneOutlinerTreeItemPtr>& InRow)
+{
+	return SNew(SContainerTreeLabel, *this, Outliner, InRow);
+}
diff --git a/Engine/Plugins/Editor/ObjectMixer/ObjectMixer/Source/ObjectMixer/Private/Views/List/RowTypes/ObjectMixerEditorListRowUObject.cpp b/Engine/Plugins/Editor/ObjectMixer/ObjectMixer/Source/ObjectMixer/Private/Views/List/RowTypes/ObjectMixerEditorListRowUObject.cpp
--- a/Engine/Plugins/Editor/ObjectMixer/ObjectMixer/Source/ObjectMixer/Private/Views/List/RowTypes/ObjectMixerEditorListRowUObject.cpp
+++ b/Engine/Plugins/Editor/ObjectMixer/ObjectMixer/Source/ObjectMixer/Private/Views/List/RowTypes/ObjectMixerEditorListRowUObject.cpp
@@ -2,9 +2,164 @@
 
 #include "Views/List/RowTypes/ObjectMixerEditorListRowUObject.h"
 
+#include "ClassIconFinder.h"
+#include "ISceneOutlinerMode.h"
+#include "Settings/EditorStyleSettings.h"
+#include "Widgets/Text/SInlineEditableTextBlock.h"
+
 const FSceneOutlinerTreeItemType FObjectMixerEditorListRowUObject::Type(&ISceneOutlinerTreeItem::Type);
 
 FString FObjectMixerEditorListRowUObject::GetDisplayString() const
 {
+	if (!RowData.GetDisplayNameOverride().IsEmptyOrWhitespace())
+	{
+		return RowData.GetDisplayNameOverride().ToString();
+	}
+	
 	return IsValid() ? ObjectSoftPtr.Get()->GetName() : FString();
 }
+
+class SObjectTreeLabel : FSceneOutlinerCommonLabelData, public SCompoundWidget
+{
+public:
+	SLATE_BEGIN_ARGS(SObjectTreeLabel) {}
+	SLATE_END_ARGS()
+
+	void Construct(const FArguments& InArgs, FObjectMixerEditorListRowUObject& ObjectItem, ISceneOutliner& SceneOutliner, const STableRow<FSceneOutlinerTreeItemPtr>& InRow)
+	{
+		WeakSceneOutliner = StaticCastSharedRef<ISceneOutliner>(SceneOutliner.AsShared());
+
+		TreeItemPtr = StaticCastSharedRef<FObjectMixerEditorListRowUObject>(ObjectItem.AsShared());
+		ObjectPtr = ObjectItem.ObjectSoftPtr.Get();
+
+		const bool bShouldUseMiddleEllipsis = GetDefault<UEditorStyleSettings>()->bEnableMiddleEllipsis;
+
+		ChildSlot
+		[
+			SNew(SHorizontalBox)
+
+			+SHorizontalBox::Slot()
+			.AutoWidth()
+			.VAlign(VAlign_Center)
+			.Padding(FSceneOutlinerDefaultTreeItemMetrics::IconPadding())
+			[
+				SNew(SBox)
+				.WidthOverride(static_cast<float>(FSceneOutlinerDefaultTreeItemMetrics::IconSize()))
+				.HeightOverride(static_cast<float>(FSceneOutlinerDefaultTreeItemMetrics::IconSize()))
+				[
+					SNew(SImage)
+					.Image(this, &SObjectTreeLabel::GetIcon)
+					.ToolTipText(this, &SObjectTreeLabel::GetIconTooltip)
+					.ColorAndOpacity(FSlateColor::UseForeground())
+				]
+			]
+
+			+SHorizontalBox::Slot()
+			.FillWidth(1.0f)
+			.VAlign(VAlign_Center)
+			.Padding(0.0f, 0.0f)
+			[
+				SNew(SHorizontalBox)
+				
+				+SHorizontalBox::Slot()
+				.VAlign(VAlign_Center)
+				[
+					SNew(STextBlock)
+					.Text(this, &SObjectTreeLabel::GetDisplayText)
+					.ToolTipText(this, &SObjectTreeLabel::GetTooltipText)
+					.HighlightText(SceneOutliner.GetFilterHighlightText())
+					.ColorAndOpacity(this, &SObjectTreeLabel::GetTextForegroundColor)
+					.OverflowPolicy(bShouldUseMiddleEllipsis ? ETextOverflowPolicy::MiddleEllipsis : TOptional<ETextOverflowPolicy>())
+				]
+			]
+		];
+	}
+
+private:
+	FText GetDisplayText() const
+	{
+		if (const FSceneOutlinerTreeItemPtr TreeItem = TreeItemPtr.Pin())
+		{
+			return FText::FromString(TreeItem->GetDisplayString());
+		}
+
+		return FText();
+	}
+
+	FText GetTooltipText() const
+	{
+		if (const FSceneOutlinerTreeItemPtr TreeItem = TreeItemPtr.Pin())
+		{
+			return FText::FromString(TreeItem->GetDisplayString());
+		}
+
+		return FText();
+	}
+
+	const FSlateBrush* GetIcon() const
+	{
+		if (const UObject* Object = ObjectPtr.Get())
+		{
+			if (WeakSceneOutliner.IsValid())
+			{
+				const FSlateBrush* CachedBrush = WeakSceneOutliner.Pin()->GetCachedIconForClass(Object->GetClass()->GetFName());
+				if (CachedBrush != nullptr)
+				{
+					return CachedBrush;
+				}
+				else
+				{
+
+					const FSlateBrush* FoundSlateBrush = FSlateIconFinder::FindIconForClass(Object->GetClass()).GetIcon();
+					WeakSceneOutliner.Pin()->CacheIconForClass(Object->GetClass()->GetFName(), FoundSlateBrush);
+					return FoundSlateBrush;
+				}
+			}
+			else
+			{
+				return nullptr;
+			}
+		}
+		else
+		{
+			return FSlateIconFinder::FindIconForClass(UObject::StaticClass()).GetOptionalIcon();
+		}
+	}
+	
+	FText GetIconTooltip() const
+	{
+		auto TreeItem = TreeItemPtr.Pin();
+		if (!TreeItem.IsValid())
+		{
+			return FText();
+		}
+
+		FText ToolTipText;
+		if (UObject* Object = ObjectPtr.Get())
+		{
+			ToolTipText = FText::FromString(Object->GetClass()->GetName());
+		}
+
+		return ToolTipText;
+	}
+
+	FSlateColor GetTextForegroundColor() const
+	{
+		auto TreeItem = TreeItemPtr.Pin();
+		if (auto BaseColor = FSceneOutlinerCommonLabelData::GetForegroundColor(*TreeItem))
+		{
+			return BaseColor.GetValue();
+		}
+		
+		return FSlateColor::UseForeground();
+	}
+	
+private:
+	TWeakPtr<FObjectMixerEditorListRowUObject> TreeItemPtr;
+	TWeakObjectPtr<UObject> ObjectPtr;
+};
+
+TSharedRef<SWidget> FObjectMixerEditorListRowUObject::GenerateLabelWidget(ISceneOutliner& Outliner, const STableRow<FSceneOutlinerTreeItemPtr>& InRow)
+{
+	return SNew(SObjectTreeLabel, *this, Outliner, InRow);
+}
diff --git a/Engine/Plugins/Editor/ObjectMixer/ObjectMixer/Source/ObjectMixer/Public/ObjectFilter/ObjectMixerEditorActorSubObject.h b/Engine/Plugins/Editor/ObjectMixer/ObjectMixer/Source/ObjectMixer/Public/ObjectFilter/ObjectMixerEditorActorSubObject.h
new file mode 100644
--- /dev/null
+++ b/Engine/Plugins/Editor/ObjectMixer/ObjectMixer/Source/ObjectMixer/Public/ObjectFilter/ObjectMixerEditorActorSubObject.h
@@ -0,0 +1,44 @@
+﻿// Copyright Epic Games, Inc. All Rights Reserved.
+
+#pragma once
+
+#include "Internationalization/Text.h"
+#include "Misc/Attribute.h"
+#include "UObject/WeakObjectPtr.h"
+
+#define UE_API OBJECTMIXEREDITOR_API
+
+struct FSlateBrush;
+
+/** Container property mixer entry for displaying a container property (e.g. Array) in the mixer */
+struct UE_API FObjectMixerEditorContainerProperty
+{
+	/**
+	 * If the subobject is stored in a container, supplying the container property's name will allow the mixer to display
+	 * a row item for the container property itself
+	 */
+	FName PropertyName = NAME_None;
+
+	/** Optional custom icon to display for the container property; if not provided, a default one will be used */
+	TAttribute<const FSlateBrush*> PropertyIconOverride;
+};
+
+/**
+ * Actor sub-object mixer entry that controls how the sub-object is displayed in the mixer
+ */
+struct UE_API FObjectMixerEditorActorSubObject
+{
+	/** Pointer to the subobject */
+	TWeakObjectPtr<UObject> Object = nullptr;
+
+	/** Optional display name to use when displaying the object in the mixer */
+	FText DisplayName = FText::GetEmpty();
+
+	/** Whether to include all parent objects up to the owning actor in the mixer */
+	bool bIncludeObjectStack = false;
+
+	/** Optional container property to make the sub-object a child of in the mixer */
+	FObjectMixerEditorContainerProperty ContainerProperty;
+};
+
+#undef UE_API
\ No newline at end of file
diff --git a/Engine/Plugins/Editor/ObjectMixer/ObjectMixer/Source/ObjectMixer/Public/ObjectFilter/ObjectMixerEditorObjectFilter.h b/Engine/Plugins/Editor/ObjectMixer/ObjectMixer/Source/ObjectMixer/Public/ObjectFilter/ObjectMixerEditorObjectFilter.h
--- a/Engine/Plugins/Editor/ObjectMixer/ObjectMixer/Source/ObjectMixer/Public/ObjectFilter/ObjectMixerEditorObjectFilter.h
+++ b/Engine/Plugins/Editor/ObjectMixer/ObjectMixer/Source/ObjectMixer/Public/ObjectFilter/ObjectMixerEditorObjectFilter.h
@@ -2,6 +2,7 @@
 
 #pragma once
 
+#include "ObjectMixerEditorActorSubObject.h"
 #include "GameFramework/Actor.h"
 #include "SceneOutlinerDragDrop.h"
 #include "UObject/Class.h"
@@ -153,6 +154,11 @@ public:
 	 */
 	UE_API virtual bool IsActorAssociated(AActor* Actor, AActor* AssociatedActor) const;
 
+	/**
+	 * Gets a list of subobjects to add to the outliner under the specified actor
+	 */
+	UE_API virtual TArray<FObjectMixerEditorActorSubObject> GetActorSubObjects(AActor* InActor) const;
+	
 	/** 
 	 * Returns true if the given tree item should override the Object Mixer's standard handling for drop operations (ValidateDrop/OnDrop)
 	 */
diff --git a/Engine/Plugins/Editor/ObjectMixer/ObjectMixer/Source/ObjectMixer/Public/SelectionInterface/IObjectMixerSelectionInterface.h b/Engine/Plugins/Editor/ObjectMixer/ObjectMixer/Source/ObjectMixer/Public/SelectionInterface/IObjectMixerSelectionInterface.h
--- a/Engine/Plugins/Editor/ObjectMixer/ObjectMixer/Source/ObjectMixer/Public/SelectionInterface/IObjectMixerSelectionInterface.h
+++ b/Engine/Plugins/Editor/ObjectMixer/ObjectMixer/Source/ObjectMixer/Public/SelectionInterface/IObjectMixerSelectionInterface.h
@@ -36,12 +36,23 @@ public:
 	 */
 	virtual void SelectComponents(const TArray<UActorComponent*>& InSelectedComponents, bool bShouldSelect, bool bSelectEvenIfHidden) = 0;
 
+	/**
+	 * Select or deselect a list of actor sub-objects.
+	 * @param	InSelectedSubObjects	The sub-objects to select.
+	 * @param	bShouldSelect			If true, select the sub-objects and deselect all other sub-objects. If false, only deselect the listed sub-objects.
+	 * @param	bSelectEvenIfHidden		If false, only select sub-objects that are visible (e.g. not filtered out) in the editor.
+	 */
+	virtual void SelectSubObjects(const TArray<UObject*>& InSelectedSubObjects, bool bShouldSelect, bool bSelectEvenIfHidden) = 0;
+	
 	/** Get the list of selected actors. */
 	virtual TArray<AActor*> GetSelectedActors() const = 0;
 
 	/** Get the list of selected components. */
 	virtual TArray<UActorComponent*> GetSelectedComponents() const = 0;
 
+	/** Get the list of selected sub-objects */
+	virtual TArray<UObject*> GetSelectedSubObjects() const = 0;
+	
 	/** Get the event handler that is raised when the synchronized selection changes. */
 	virtual FOnSelectionChanged& OnSelectionChanged() = 0;
 };
diff --git a/Engine/Plugins/Editor/ObjectMixer/ObjectMixer/Source/ObjectMixer/Public/SelectionInterface/LevelEditorObjectMixerSelectionInterface.h b/Engine/Plugins/Editor/ObjectMixer/ObjectMixer/Source/ObjectMixer/Public/SelectionInterface/LevelEditorObjectMixerSelectionInterface.h
--- a/Engine/Plugins/Editor/ObjectMixer/ObjectMixer/Source/ObjectMixer/Public/SelectionInterface/LevelEditorObjectMixerSelectionInterface.h
+++ b/Engine/Plugins/Editor/ObjectMixer/ObjectMixer/Source/ObjectMixer/Public/SelectionInterface/LevelEditorObjectMixerSelectionInterface.h
@@ -19,8 +19,10 @@ public:
 	//~ Begin IObjectMixerSelectionInterface interface
 	UE_API virtual void SelectActors(const TArray<AActor*>& InSelectedActors, bool bShouldSelect, bool bSelectEvenIfHidden) override;
 	UE_API virtual void SelectComponents(const TArray<UActorComponent*>& InSelectedComponents, bool bShouldSelect, bool bSelectEvenIfHidden) override;
+	UE_API virtual void SelectSubObjects(const TArray<UObject*>& InSelectedSubObjects, bool bShouldSelect, bool bSelectEvenIfHidden) override { }
 	UE_API virtual TArray<AActor*> GetSelectedActors() const override;
 	UE_API virtual TArray<UActorComponent*> GetSelectedComponents() const override;
+	UE_API virtual TArray<UObject*> GetSelectedSubObjects() const override { return { }; }
 	virtual FOnSelectionChanged& OnSelectionChanged() override { return SelectionChanged; }
 	//~ End IObjectMixerSelectionInterface interface
 
diff --git a/Engine/Plugins/Editor/ObjectMixer/ObjectMixer/Source/ObjectMixer/Public/Views/List/Modes/ObjectMixerOutlinerHierarchy.h b/Engine/Plugins/Editor/ObjectMixer/ObjectMixer/Source/ObjectMixer/Public/Views/List/Modes/ObjectMixerOutlinerHierarchy.h
--- a/Engine/Plugins/Editor/ObjectMixer/ObjectMixer/Source/ObjectMixer/Public/Views/List/Modes/ObjectMixerOutlinerHierarchy.h
+++ b/Engine/Plugins/Editor/ObjectMixer/ObjectMixer/Source/ObjectMixer/Public/Views/List/Modes/ObjectMixerOutlinerHierarchy.h
@@ -57,6 +57,14 @@ protected:
 	 */
 	bool CreateAssociatedActorItems(AActor* Actor, TArray<FSceneOutlinerTreeItemPtr>& OutItems) const;
 
+	/**
+	 * Create all subobject items for an actor and place them in OutItems.
+	 * @param Actor The actor to get subobjects for
+	 * @param OutItems The resulting list of tree items
+	 * @return true if any items would have been returned before filtering
+	 */
+	bool CreateSubObjectItems(AActor* Actor, TArray<FSceneOutlinerTreeItemPtr>& OutItems) const;
+	
 	/**
 	 * Create all component items for an actor if we are showing components and place them in OutItems. Items not added if filtered out.
 	 * If components are returned or would be returned before filters, the Actor will be added at the beginning of the array.
diff --git a/Engine/Plugins/Editor/ObjectMixer/ObjectMixer/Source/ObjectMixer/Public/Views/List/Modes/ObjectMixerOutlinerMode.h b/Engine/Plugins/Editor/ObjectMixer/ObjectMixer/Source/ObjectMixer/Public/Views/List/Modes/ObjectMixerOutlinerMode.h
--- a/Engine/Plugins/Editor/ObjectMixer/ObjectMixer/Source/ObjectMixer/Public/Views/List/Modes/ObjectMixerOutlinerMode.h
+++ b/Engine/Plugins/Editor/ObjectMixer/ObjectMixer/Source/ObjectMixer/Public/Views/List/Modes/ObjectMixerOutlinerMode.h
@@ -68,6 +68,11 @@ namespace ObjectMixerOutliner
 	{
 		bool operator()(TWeakPtr<ISceneOutlinerTreeItem> Item, FFolder& DataOut) const;
 	};
+
+	struct FObjectSelector
+	{
+		bool operator()(TWeakPtr<ISceneOutlinerTreeItem> Item, UObject*& OutObject);
+	};
 }
 
 USTRUCT()
@@ -193,6 +198,7 @@ public:
 	virtual bool IsInteractive() const override { return true; }
 	UE_API virtual TUniquePtr<ISceneOutlinerHierarchy> CreateHierarchy() override;
 	UE_API virtual void Rebuild() override;
+	UE_API virtual int32 GetTypeSortPriority(const ISceneOutlinerTreeItem& Item) const override;
 	UE_API virtual FCreateSceneOutlinerMode CreateFolderPickerMode(const FFolder::FRootObject& InRootObject = FFolder::GetInvalidRootObject()) const override;
 	UE_API virtual void CreateViewContent(FMenuBuilder& MenuBuilder) override;
 	UE_API virtual TSharedPtr<FDragDropOperation> CreateDragDropOperation(const FPointerEvent& MouseEvent, const TArray<FSceneOutlinerTreeItemPtr>& InTreeItems) const override;
@@ -312,6 +318,7 @@ protected:
 	UE_API bool HasComponentSelectionChanged(TArray<UActorComponent*>& OutSelectedComponents, bool& bOutAreAnyInPIE);
 	UE_API void SelectActorsInEditor(const TArray<AActor*>& InSelectedActors, bool bShouldSelect, bool bSelectEvenIfHidden);
 	UE_API void SelectComponentsInEditor(const TArray<UActorComponent*>& InSelectedComponents, bool bShouldSelect, bool bSelectEvenIfHidden);
+	UE_API void SelectSubObjectsInEditor(const TArray<UObject*>& InSelectedSubObjects, bool bShouldSelect, bool bSelectEvenIfHidden);
 	UE_API void SelectActorsInMixer(const TArray<AActor*>& InSelectedActors, bool bShouldSelect, bool bSelectEvenIfHidden);
 	UE_API void SelectComponentsInMixer(const TArray<UActorComponent*>& InSelectedComponents, bool bShouldSelect, bool bSelectEvenIfHidden);
 
@@ -321,6 +328,9 @@ protected:
 	/** Get the list of components selected in the editor */
 	UE_API TArray<UActorComponent*> GetSelectedComponentsInEditor() const;
 
+	/** Gets the list of actor sub-objects that are selected in the editor. */
+	UE_API TArray<UObject*> GetSelectedSubObjectsInEditor() const;
+
 	/** Check if an object is a member of one of the classes our list model is filtered to */
 	UE_API bool GetIsObjectOfFilteredClass(const UObject* Object) const;
 	
@@ -337,6 +347,9 @@ protected:
 	UE_API void SynchronizeComponentAndActorSelection();
 	UE_API void SynchronizeSelectedActorDescs();
 
+	/** Synchronizes the object mixer's selection with any sub-objects selected elsewhere in the editor */
+	UE_API void SynchronizeSubObjectSelection();
+
 	UE_API void OnActorEditorContextSubsystemChanged();
 
 	/** Filter factories */
diff --git a/Engine/Plugins/Editor/ObjectMixer/ObjectMixer/Source/ObjectMixer/Public/Views/List/RowTypes/ObjectMixerEditorListRowContainer.h b/Engine/Plugins/Editor/ObjectMixer/ObjectMixer/Source/ObjectMixer/Public/Views/List/RowTypes/ObjectMixerEditorListRowContainer.h
new file mode 100644
--- /dev/null
+++ b/Engine/Plugins/Editor/ObjectMixer/ObjectMixer/Source/ObjectMixer/Public/Views/List/RowTypes/ObjectMixerEditorListRowContainer.h
@@ -0,0 +1,45 @@
+﻿// Copyright Epic Games, Inc. All Rights Reserved.
+
+#pragma once
+
+#include "ISceneOutlinerTreeItem.h"
+
+#define UE_API OBJECTMIXEREDITOR_API
+
+/** Object Mixer row that represents container properties such as arrays that contain subobjects */
+struct FObjectMixerEditorListRowContainer : ISceneOutlinerTreeItem
+{
+public:
+	explicit FObjectMixerEditorListRowContainer(UObject* InPropertyOwner, FName InPropertyName, SSceneOutliner* InSceneOutliner);
+	
+	/* Begin ISceneOutlinerTreeItem Implementation */
+	static UE_API const FSceneOutlinerTreeItemType Type;
+	UE_API virtual FString GetDisplayString() const override;
+	virtual bool CanInteract() const override { return true; }
+	virtual bool IsValid() const override { return !PropertyOwner.IsNull() && Property != nullptr; }
+	virtual FSceneOutlinerTreeItemID GetID() const override { return IsValid() ? UniqueId : FSceneOutlinerTreeItemID(); }
+	virtual TSharedRef<SWidget> GenerateLabelWidget(ISceneOutliner& Outliner, const STableRow<FSceneOutlinerTreeItemPtr>& InRow) override;
+	/* End ISceneOutlinerTreeItem Implementation */
+
+	/** Gets the object that owns the container property being represented by this row */
+	UObject* GetPropertyOwner() const { return PropertyOwner.Get(); }
+
+	/** Gets the FProperty being represented by this row */
+	const FProperty* GetContainerProperty() const { return Property; }
+
+public:
+	/** Optional icon to display in the list row; if not provided, a default icon will be used */
+	TAttribute<const FSlateBrush*> PropertyIconOverride;
+	
+private:
+	/** The explicit object whose container property is being represented by this list row */
+	TSoftObjectPtr<UObject> PropertyOwner;
+	
+	/** The FProperty that is being represented by this list row */
+	FProperty* Property = nullptr;
+
+	/** The unique ID of this list row in the scene outliner */
+	uint64 UniqueId = 0;
+};
+
+#undef UE_API
\ No newline at end of file
diff --git a/Engine/Plugins/Editor/ObjectMixer/ObjectMixer/Source/ObjectMixer/Public/Views/List/RowTypes/ObjectMixerEditorListRowUObject.h b/Engine/Plugins/Editor/ObjectMixer/ObjectMixer/Source/ObjectMixer/Public/Views/List/RowTypes/ObjectMixerEditorListRowUObject.h
--- a/Engine/Plugins/Editor/ObjectMixer/ObjectMixer/Source/ObjectMixer/Public/Views/List/RowTypes/ObjectMixerEditorListRowUObject.h
+++ b/Engine/Plugins/Editor/ObjectMixer/ObjectMixer/Source/ObjectMixer/Public/Views/List/RowTypes/ObjectMixerEditorListRowUObject.h
@@ -32,12 +32,19 @@ struct FObjectMixerEditorListRowUObject : ISceneOutlinerTreeItem
 	/** Constant identifier for this tree item */
 	const FObjectKey ID;
 
+	/** If the UObject list row should be displayed as part of a container, this is the container's property name in this object's outer */
+	FName ContainerProperty = NAME_None;
+
+	/** If true, the mixer displays rows for all of this object's parents up to its owning actor; otherwise, it is parented directly to its owning actor */
+	bool bIncludeObjectStack = false;
+	
 	/* Begin ISceneOutlinerTreeItem Implementation */
 	static UE_API const FSceneOutlinerTreeItemType Type;
 	UE_API virtual FString GetDisplayString() const override;
 	virtual bool CanInteract() const override { return true; }
 	virtual bool IsValid() const override { return !ObjectSoftPtr.IsNull(); }
 	virtual FSceneOutlinerTreeItemID GetID() const override { return ID; }
+	virtual TSharedRef<SWidget> GenerateLabelWidget(ISceneOutliner& Outliner, const STableRow<FSceneOutlinerTreeItemPtr>& InRow) override;
 	/* End ISceneOutlinerTreeItem Implementation */
 };
 
diff --git a/Engine/Plugins/Runtime/nDisplay/Source/DisplayClusterColorGrading/Private/SelectionInterface/DisplayClusterObjectMixerSelectionInterface.h b/Engine/Plugins/Runtime/nDisplay/Source/DisplayClusterColorGrading/Private/SelectionInterface/DisplayClusterObjectMixerSelectionInterface.h
--- a/Engine/Plugins/Runtime/nDisplay/Source/DisplayClusterColorGrading/Private/SelectionInterface/DisplayClusterObjectMixerSelectionInterface.h
+++ b/Engine/Plugins/Runtime/nDisplay/Source/DisplayClusterColorGrading/Private/SelectionInterface/DisplayClusterObjectMixerSelectionInterface.h
@@ -17,8 +17,10 @@ public:
 	//~ Begin IObjectMixerSelectionInterface interface
 	virtual void SelectActors(const TArray<AActor*>& InSelectedActors, bool bShouldSelect, bool bSelectEvenIfHidden) override;
 	virtual void SelectComponents(const TArray<UActorComponent*>& InSelectedComponents, bool bShouldSelect, bool bSelectEvenIfHidden) override;
+	virtual void SelectSubObjects(const TArray<UObject*>& InSelectedSubObjects, bool bShouldSelect, bool bSelectEvenIfHidden) override { }
 	virtual TArray<AActor*> GetSelectedActors() const override;
 	virtual TArray<UActorComponent*> GetSelectedComponents() const override;
+	virtual TArray<UObject*> GetSelectedSubObjects() const override { return { }; }
 	virtual FOnSelectionChanged& OnSelectionChanged() override { return SelectionChanged; }
 	//~ End IObjectMixerSelectionInterface interface
 
diff --git a/Engine/Source/Editor/DetailCustomizations/Private/ColorGradingSettingsCustomization.cpp b/Engine/Source/Editor/DetailCustomizations/Private/ColorGradingSettingsCustomization.cpp
--- a/Engine/Source/Editor/DetailCustomizations/Private/ColorGradingSettingsCustomization.cpp
+++ b/Engine/Source/Editor/DetailCustomizations/Private/ColorGradingSettingsCustomization.cpp
@@ -37,7 +37,7 @@ void FColorGradingSettingsCustomization::CustomizeChildren(TSharedRef<IPropertyH
 	for (uint32 Index = 0; Index < NumChildProperties; ++Index)
 	{
 		TSharedPtr<IPropertyHandle> ChildHandle = PropertyHandle->GetChildHandle(Index);
-		if (!ChildHandle.IsValid() || !ChildHandle->IsValidHandle())
+		if (!ChildHandle.IsValid() || !ChildHandle->IsValidHandle() || ChildHandle->IsCustomized())
 		{
 			continue;
 		}
