Composure: Fixed issue where editor would crash if a new level was loaded while the color grading drawer was displaying a Composite actor color grading pass.

#Synced-CL 49143073

diff --git a/Engine/Plugins/Editor/ObjectMixer/ObjectMixer/Source/ObjectMixer/Private/Views/List/RowTypes/ObjectMixerEditorListRowUObject.cpp b/Engine/Plugins/Editor/ObjectMixer/ObjectMixer/Source/ObjectMixer/Private/Views/List/RowTypes/ObjectMixerEditorListRowUObject.cpp
--- a/Engine/Plugins/Editor/ObjectMixer/ObjectMixer/Source/ObjectMixer/Private/Views/List/RowTypes/ObjectMixerEditorListRowUObject.cpp
+++ b/Engine/Plugins/Editor/ObjectMixer/ObjectMixer/Source/ObjectMixer/Private/Views/List/RowTypes/ObjectMixerEditorListRowUObject.cpp
@@ -16,7 +16,7 @@ FString FObjectMixerEditorListRowUObject::GetDisplayString() const
 		return RowData.GetDisplayNameOverride().ToString();
 	}
 	
-	return IsValid() ? ObjectSoftPtr.Get()->GetName() : FString();
+	return IsValid() && ObjectSoftPtr.IsValid() ? ObjectSoftPtr.Get()->GetName() : FString();
 }
 
 class SObjectTreeLabel : FSceneOutlinerCommonLabelData, public SCompoundWidget
